PBIK

diff --git a/Engine/Plugins/Experimental/Animation/PBIK/PBIK.uplugin b/Engine/Plugins/Experimental/Animation/PBIK/PBIK.uplugin
new file mode 100644
--- /dev/null
+++ b/Engine/Plugins/Experimental/Animation/PBIK/PBIK.uplugin
@@ -0,0 +1,31 @@
+{
+	"FileVersion": 3,
+	"Version": 0.1,
+	"VersionName": "0.1",
+	"FriendlyName": "PBIK",
+	"Description": "",
+	"Category": "Other",
+	"CreatedBy": "",
+	"CreatedByURL": "",
+	"DocsURL": "",
+	"MarketplaceURL": "",
+	"SupportURL": "",
+	"EnabledByDefault": false,
+	"CanContainContent": true,
+	"IsBetaVersion": false,
+	"IsExperimentalVersion": true,
+	"Installed": false,
+	"Plugins": [
+		{
+			"Name": "ControlRig",
+			"Enabled": true
+		}
+	],
+	"Modules": [
+		{
+			"Name": "PBIK",
+			"Type": "Runtime",
+			"LoadingPhase": "Default"
+		}
+	]
+}
\ No newline at end of file
diff --git a/Engine/Plugins/Experimental/Animation/PBIK/Source/PBIK/PBIK.Build.cs b/Engine/Plugins/Experimental/Animation/PBIK/Source/PBIK/PBIK.Build.cs
new file mode 100644
--- /dev/null
+++ b/Engine/Plugins/Experimental/Animation/PBIK/Source/PBIK/PBIK.Build.cs
@@ -0,0 +1,30 @@
+// Copyright Epic Games, Inc. All Rights Reserved.
+
+using UnrealBuildTool;
+
+public class PBIK : ModuleRules
+{
+	public PBIK(ReadOnlyTargetRules Target) : base(Target)
+	{
+		PCHUsage = ModuleRules.PCHUsageMode.UseExplicitOrSharedPCHs;
+
+		PublicDependencyModuleNames.AddRange(
+			new string[]
+			{
+				"Core",
+				"RigVM",
+				"ControlRig"
+			}
+			);
+
+		PrivateDependencyModuleNames.AddRange(
+			new string[]
+			{
+				"CoreUObject",
+				"Engine",
+			}
+			);
+
+		
+	}
+}
diff --git a/Engine/Plugins/Experimental/Animation/PBIK/Source/PBIK/Private/Core/PBIKBody.cpp b/Engine/Plugins/Experimental/Animation/PBIK/Source/PBIK/Private/Core/PBIKBody.cpp
new file mode 100644
--- /dev/null
+++ b/Engine/Plugins/Experimental/Animation/PBIK/Source/PBIK/Private/Core/PBIKBody.cpp
@@ -0,0 +1,129 @@
+// Copyright Epic Games, Inc. All Rights Reserved.
+
+#pragma once
+
+#include "Core/PBIKBody.h"
+#include "Core/PBIKSolver.h"
+
+namespace PBIK
+{
+
+FBone::FBone(
+	const FName InName,
+	const int& InParentIndex,		// must pass -1 for root of whole skeleton
+	const FVector& InOrigPosition,
+	const FQuat& InOrigRotation,
+	bool bInIsSolverRoot)
+{
+	Name = InName;
+	ParentIndex = InParentIndex;
+	Position = InOrigPosition;
+	Rotation = InOrigRotation;
+	bIsSolverRoot = bInIsSolverRoot;
+	bIsSolved = false; // default
+}
+
+bool FBone::HasChild(const FBone* Bone)
+{
+	for (const FBone* Child : Children)
+	{
+		if (Bone->Name == Child->Name)
+		{
+			return true;
+		}
+	}
+
+	return false;
+}
+
+FRigidBody::FRigidBody(FBone* InBone)
+{
+	Bone = InBone;
+	J = FBoneSettings();
+}
+
+void FRigidBody::Initialize(FBone* SolverRoot)
+{
+	FVector Centroid = Bone->Position;
+	Length = 0.0f;
+	for(const FBone* Child : Bone->Children)
+	{
+		Centroid += Child->Position;
+		Length += (Bone->Position - Child->Position).Size();
+	}
+	Centroid = Centroid * (1.0f / ((float)Bone->Children.Num() + 1.0f));
+
+	Position = Centroid;
+	Rotation = RotationOrig = Bone->Rotation;
+	BoneLocalPosition = Rotation.Inverse() * (Bone->Position - Centroid);
+
+	for (FBone* Child : Bone->Children)
+	{
+		FVector ChildLocalPos = Rotation.Inverse() * (Child->Position - Centroid);
+		ChildLocalPositions.Add(ChildLocalPos);
+	}
+
+	// calculate num bones distance to root
+	NumBonesToRoot = 0;
+	FBone* Parent = Bone;
+	while (Parent && Parent != SolverRoot)
+	{
+		NumBonesToRoot += 1;
+		Parent = Parent->Parent;
+	}
+}
+
+void FRigidBody::UpdateFromInputs(const FPBIKSolverSettings& Settings)
+{
+	// set to input pose
+	Position = Bone->Position - Bone->Rotation * BoneLocalPosition;
+	Rotation = Bone->Rotation;
+
+	// Body.Length used as rough approximation of the mass of the body
+	// for fork joints (multiple solved children) we sum lengths to all children (see Initialize)
+	InvMass = 1.0f / ( Length * ((Settings.MassMultiplier * GLOBAL_UNITS) + 0.5f));
+}
+
+int FRigidBody::GetNumBonesToRoot() const
+{ 
+	return NumBonesToRoot; 
+}
+
+FRigidBody* FRigidBody::GetParentBody()
+{
+	if (Bone && Bone->Parent)
+	{
+		return Bone->Parent->Body;
+	}
+
+	return nullptr;
+}
+
+void FRigidBody::ApplyPushToRotateBody(const FVector& Push, const FVector& Offset)
+{
+	// quat differentiation
+	FVector Omega = InvMass * (1.0f - J.RotationStiffness) * FVector::CrossProduct(Offset, Push);
+	FQuat OQ(Omega.X, Omega.Y, Omega.Z, 0.0f);
+	OQ = OQ * Rotation;
+	OQ.X *= 0.5f;
+	OQ.Y *= 0.5f;
+	OQ.Z *= 0.5f;
+	OQ.W *= 0.5f;
+	Rotation.X = Rotation.X + OQ.X;
+	Rotation.Y = Rotation.Y + OQ.Y;
+	Rotation.Z = Rotation.Z + OQ.Z;
+	Rotation.W = Rotation.W + OQ.W;
+	Rotation.Normalize();
+}
+
+void FRigidBody::ApplyPushToPosition(const FVector& Push)
+{
+	if (bPinnedToEffector)
+	{
+		return; // pins are locked
+	}
+
+	Position += Push * (1.0f - J.PositionStiffness);
+}
+
+} // namespace
\ No newline at end of file
diff --git a/Engine/Plugins/Experimental/Animation/PBIK/Source/PBIK/Private/Core/PBIKConstraint.cpp b/Engine/Plugins/Experimental/Animation/PBIK/Source/PBIK/Private/Core/PBIKConstraint.cpp
new file mode 100644
--- /dev/null
+++ b/Engine/Plugins/Experimental/Animation/PBIK/Source/PBIK/Private/Core/PBIKConstraint.cpp
@@ -0,0 +1,325 @@
+// Copyright Epic Games, Inc. All Rights Reserved.
+
+#pragma once
+
+#include "Core/PBIKConstraint.h"
+#include "Core/PBIKSolver.h"
+#include "Math/UnrealMathUtility.h"
+
+//#pragma optimize("", off)
+
+namespace PBIK
+{
+	
+FJointConstraint::FJointConstraint(FRigidBody* InA, FRigidBody* InB)
+{
+	A = InA;
+	B = InB;
+
+	FVector PinPoint = B->Bone->Position;
+	PinPointLocalToA = A->Rotation.Inverse() * (PinPoint - A->Position);
+	PinPointLocalToB = B->Rotation.Inverse() * (PinPoint - B->Position);
+
+	XOrig = B->RotationOrig * FVector(1, 0, 0);
+	YOrig = B->RotationOrig * FVector(0, 1, 0);
+	ZOrig = B->RotationOrig * FVector(0, 0, 1);
+}
+	
+void FJointConstraint::Solve(bool bMoveSubRoots)
+{
+	// get pos correction to use for rotation
+	FVector OffsetA;
+	FVector OffsetB;
+	FVector Correction = GetPositionCorrection(OffsetA, OffsetB);
+
+	// apply rotation from correction 
+	// at pin point to both bodies
+	A->ApplyPushToRotateBody(Correction, OffsetA);
+	B->ApplyPushToRotateBody(-Correction, OffsetB);
+
+	// enforce joint limits
+	UpdateJointLimits();
+
+	// apply positional correction to align pin point on both Bodies
+	// NOTE: applying position correction AFTER rotation takes into
+	// consideration change in relative pin locations mid-step after
+	// bodies have been rotated by ApplyPushToRotateBody
+	Correction = GetPositionCorrection(OffsetA, OffsetB);
+
+	// scale correction by inv mass
+	bool bALocked = A->bPinnedToEffector || (!bMoveSubRoots && A->Bone->bIsSubRoot);
+	bool bBLocked = B->bPinnedToEffector || (!bMoveSubRoots && B->Bone->bIsSubRoot);
+	float AInvMass = bALocked ? 0.0f : 1.0f;
+	float BInvMass = bBLocked ? 0.0f : 1.0f;
+	float InvMassSum = AInvMass + BInvMass;
+	if (InvMassSum < KINDA_SMALL_NUMBER)
+	{
+		return;
+	}
+	Correction = Correction * (1.0f / InvMassSum);
+
+	A->ApplyPushToPosition(Correction * AInvMass);
+	B->ApplyPushToPosition(-Correction * BInvMass);
+}
+
+void FJointConstraint::RemoveStretch()
+{
+	// apply position correction, keeping parent body fixed
+	FVector ToA;
+	FVector ToB;
+	FVector Correction = GetPositionCorrection(ToA, ToB);
+	B->Position -= Correction;
+}
+
+FVector FJointConstraint::GetPositionCorrection(FVector& OutBodyToA, FVector& OutBodyToB)
+{
+	OutBodyToA = A->Rotation * PinPointLocalToA;
+	OutBodyToB = B->Rotation * PinPointLocalToB;
+	FVector PinPointOnA = A->Position + OutBodyToA;
+	FVector PinPointOnB = B->Position + OutBodyToB;
+	return (PinPointOnB - PinPointOnA);
+}
+
+void FJointConstraint::ApplyRotationCorrection(FQuat PureRotA, FQuat PureRotB)
+{
+	// https://matthias-research.github.io/pages/publications/PBDBodies.pdf/
+	// Equation 8 and 9 from "Detailed Rigid Body Simulation with Extended Position Based Dynamics"
+
+	PureRotA.X *= 0.5f;
+	PureRotA.Y *= 0.5f;
+	PureRotA.Z *= 0.5f;
+	PureRotB.X *= 0.5f;
+	PureRotB.Y *= 0.5f;
+	PureRotB.Z *= 0.5f;
+	PureRotA = PureRotA * A->Rotation;
+	PureRotB = PureRotB * B->Rotation;
+
+	A->Rotation.X = A->Rotation.X + PureRotA.X;
+	A->Rotation.Y = A->Rotation.Y + PureRotA.Y;
+	A->Rotation.Z = A->Rotation.Z + PureRotA.Z;
+	A->Rotation.W = A->Rotation.W + PureRotA.W;
+
+	B->Rotation.X = B->Rotation.X - PureRotB.X;
+	B->Rotation.Y = B->Rotation.Y - PureRotB.Y;
+	B->Rotation.Z = B->Rotation.Z - PureRotB.Z;
+	B->Rotation.W = B->Rotation.W - PureRotB.W;
+
+	A->Rotation.Normalize();
+	B->Rotation.Normalize();
+}
+
+void FJointConstraint::UpdateJointLimits()
+{
+	FBoneSettings& J = B->J;
+
+	// no limits at all
+	if (J.X == ELimitType::Free &&
+		J.Y == ELimitType::Free &&
+		J.Z == ELimitType::Free)
+	{
+		return;
+	}
+
+	// force max angles > min angles
+	J.MaxX = J.MaxX < J.MinX ? J.MinX + 1 : J.MaxX;
+	J.MaxY = J.MaxY < J.MinY ? J.MinY + 1 : J.MaxY;
+	J.MaxZ = J.MaxZ < J.MinZ ? J.MinZ + 1 : J.MaxZ;
+
+	// determine which axes are explicitly or implicitly locked (limited with very small allowable angle)
+	bool bLockX = (J.X == ELimitType::Locked) || (J.X == ELimitType::Limited && ((J.MaxX - J.MinX) < 2.0f));
+	bool bLockY = (J.Y == ELimitType::Locked) || (J.Y == ELimitType::Limited && ((J.MaxY - J.MinY) < 2.0f));
+	bool bLockZ = (J.Z == ELimitType::Locked) || (J.Z == ELimitType::Limited && ((J.MaxZ - J.MinZ) < 2.0f));
+
+	// determine which axes should be treated as a hinge (mutually exclusive)
+	bool bXHinge = !bLockX && bLockY && bLockZ;
+	bool bYHinge = bLockX && !bLockY && bLockZ;
+	bool bZHinge = bLockX && bLockY && !bLockZ;
+
+	// apply hinge corrections
+	if (bXHinge)
+	{
+		UpdateLocalRotateAxes(true, false, false);
+		FVector CrossProd = FVector::CrossProduct(XA, XB);
+		FQuat PureRotA = FQuat(CrossProd.X, CrossProd.Y, CrossProd.Z, 0.0f);
+		ApplyRotationCorrection(PureRotA, PureRotA);
+	}
+	else if (bYHinge)
+	{
+		UpdateLocalRotateAxes(false, true, false);
+		FVector CrossProd = FVector::CrossProduct(YA, YB);
+		FQuat PureRotA = FQuat(CrossProd.X, CrossProd.Y, CrossProd.Z, 0.0f);
+		ApplyRotationCorrection(PureRotA, PureRotA);
+	}
+	else if (bZHinge)
+	{
+		UpdateLocalRotateAxes(false, false, true);
+		FVector CrossProd = FVector::CrossProduct(ZA, ZB);
+		FQuat PureRotA = FQuat(CrossProd.X, CrossProd.Y, CrossProd.Z, 0.0f);
+		ApplyRotationCorrection(PureRotA, PureRotA);
+	}
+
+	if (bLockX || bLockY || bLockZ)
+	{
+		DecomposeRotationAngles(); // TODO don't strictly need to update everything in here
+	}
+
+	if (bLockX)
+	{
+		RotateWithinLimits(0, 0, AngleX, XA, ZBProjOnX, ZA);
+	}
+
+	if (bLockY)
+	{
+		RotateWithinLimits(0, 0, AngleY, YA, ZBProjOnY, ZA);
+	}
+
+	if (bLockZ)
+	{
+		RotateWithinLimits(0, 0, AngleZ, ZA, YBProjOnZ, YA);
+	}
+
+	// enforce min/max angles
+	bool bLimitX = J.X == ELimitType::Limited && !bLockX;
+	bool bLimitY = J.Y == ELimitType::Limited && !bLockY;
+	bool bLimitZ = J.Z == ELimitType::Limited && !bLockZ;
+	if (bLimitX || bLimitY || bLimitZ)
+	{
+		DecomposeRotationAngles();
+	}
+
+	if (bLimitX)
+	{
+		RotateWithinLimits(J.MinX, J.MaxX, AngleX, XA, ZBProjOnX, ZA);
+	}
+
+	if (bLimitY)
+	{
+		RotateWithinLimits(J.MinY, J.MaxY, AngleY, YA, ZBProjOnY, ZA);
+	}
+
+	if (bLimitZ)
+	{
+		RotateWithinLimits(J.MinZ, J.MaxZ, AngleZ, ZA, YBProjOnZ, YA);
+	}
+}
+
+void FJointConstraint::RotateWithinLimits(
+	float MinAngle,
+	float MaxAngle,
+	float CurrentAngle,
+	FVector RotAxis,
+	FVector CurVec,
+	FVector RefVec)
+{
+	bool bBeyondMin = CurrentAngle < MinAngle;
+	bool bBeyondMax = CurrentAngle > MaxAngle;
+	if (bBeyondMin || bBeyondMax)
+	{
+		float TgtAngle = bBeyondMin ? MinAngle : MaxAngle;
+		FQuat TgtRot = FQuat(RotAxis, FMath::DegreesToRadians(TgtAngle));
+		FVector TgtVec = TgtRot * RefVec;
+		FVector TgtCross = FVector::CrossProduct(TgtVec, CurVec);
+		FQuat PureRot = FQuat(TgtCross.X, TgtCross.Y, TgtCross.Z, 0.0f);
+		ApplyRotationCorrection(PureRot, PureRot);
+	}
+}
+
+void FJointConstraint::UpdateLocalRotateAxes(bool bX, bool bY, bool bZ)
+{
+	FQuat ARot = A->Rotation * A->RotationOrig.Inverse();
+	FQuat BRot = B->Rotation * B->RotationOrig.Inverse();
+
+	if (bX)
+	{
+		XA = ARot * XOrig;
+		XB = BRot * XOrig;
+	}
+
+	if (bY)
+	{
+		YA = ARot * YOrig;
+		YB = BRot * YOrig;
+	}
+
+	if (bZ)
+	{
+		ZA = ARot * ZOrig;
+		ZB = BRot * ZOrig;
+	}
+}
+
+void FJointConstraint::DecomposeRotationAngles()
+{
+	FQuat ARot = A->Rotation * A->RotationOrig.Inverse();
+	FQuat BRot = B->Rotation * B->RotationOrig.Inverse();
+
+	XA = ARot * XOrig;
+	YA = ARot * YOrig;
+	ZA = ARot * ZOrig;
+	XB = BRot * XOrig;
+	YB = BRot * YOrig;
+	ZB = BRot * ZOrig;
+
+	ZBProjOnX = FVector::VectorPlaneProject(ZB, XA).GetSafeNormal();
+	ZBProjOnY = FVector::VectorPlaneProject(ZB, YA).GetSafeNormal();
+	YBProjOnZ = FVector::VectorPlaneProject(YB, ZA).GetSafeNormal();
+
+	AngleX = SignedAngleBetweenNormals(ZA, ZBProjOnX, XA);
+	AngleY = SignedAngleBetweenNormals(ZA, ZBProjOnY, YA);
+	AngleZ = SignedAngleBetweenNormals(YA, YBProjOnZ, ZA);
+}
+
+float FJointConstraint::SignedAngleBetweenNormals(
+	const FVector& From, 
+	const FVector& To, 
+	const FVector& Axis) const
+{
+	float FromDotTo = FVector::DotProduct(From, To);
+	float Angle = FMath::RadiansToDegrees(FMath::Acos(FromDotTo));
+	FVector Cross = FVector::CrossProduct(From, To); // TODO may be backwards?
+	float Dot = FVector::DotProduct(Cross, Axis);
+	return Dot >= 0 ? Angle : -Angle;
+}
+	
+FPinConstraint::FPinConstraint(FRigidBody* InBody, const FVector& InPinPoint)
+{
+	A = InBody;
+	PinPointLocalToA = A->Rotation.Inverse() * (InPinPoint - A->Position);
+	GoalPoint = InPinPoint;
+}
+
+void FPinConstraint::Solve(bool bMoveSubRoots)
+{
+	if (!bEnabled || (Alpha <= KINDA_SMALL_NUMBER))
+	{
+		return;
+	}
+
+	if (!bMoveSubRoots && A->Bone->bIsSubRoot)
+	{
+		return;
+	}
+
+	// get positional correction for rotation
+	FVector AToPinPoint;
+	FVector Correction = GetPositionCorrection(AToPinPoint);
+
+	// rotate body from alignment of pin points
+	A->ApplyPushToRotateBody(Correction, AToPinPoint);
+
+	// apply positional correction to Body to align with target (after rotation)
+	// (applying directly without considering PositionStiffness because PinConstraints need
+	// to precisely pull the attached body to achieve convergence)
+	Correction = GetPositionCorrection(AToPinPoint);
+	A->Position += Correction; 
+}
+
+FVector FPinConstraint::GetPositionCorrection(FVector& OutBodyToPinPoint) const
+{
+	OutBodyToPinPoint = A->Rotation * PinPointLocalToA;
+	FVector PinPoint = A->Position + OutBodyToPinPoint;
+	return (GoalPoint - PinPoint) * Alpha;
+}
+
+} // namespace
+
diff --git a/Engine/Plugins/Experimental/Animation/PBIK/Source/PBIK/Private/Core/PBIKDebug.cpp b/Engine/Plugins/Experimental/Animation/PBIK/Source/PBIK/Private/Core/PBIKDebug.cpp
new file mode 100644
--- /dev/null
+++ b/Engine/Plugins/Experimental/Animation/PBIK/Source/PBIK/Private/Core/PBIKDebug.cpp
@@ -0,0 +1,24 @@
+// Copyright Epic Games, Inc. All Rights Reserved.
+
+#pragma once
+
+#include "Core/PBIKDebug.h"
+#include "Core/PBIKSolver.h"
+#include "Core/PBIKBody.h"
+
+namespace PBIK
+{
+	void FDebugDraw::GetDebugLinesForBodies(TArray<FDebugLine>& OutLines)
+	{
+		OutLines.Empty();
+		for (const FRigidBody& Body : Solver->Bodies)
+		{
+			if (Body.Bone->Children.Num() == 1)
+			{
+				FVector StartPos = Body.Position + Body.Rotation * Body.BoneLocalPosition;
+				FVector EndPos = Body.Position + Body.Rotation * Body.ChildLocalPositions[0];
+				OutLines.Emplace(StartPos, EndPos);
+			}
+		}
+	}
+}
diff --git a/Engine/Plugins/Experimental/Animation/PBIK/Source/PBIK/Private/Core/PBIKSolver.cpp b/Engine/Plugins/Experimental/Animation/PBIK/Source/PBIK/Private/Core/PBIKSolver.cpp
new file mode 100644
--- /dev/null
+++ b/Engine/Plugins/Experimental/Animation/PBIK/Source/PBIK/Private/Core/PBIKSolver.cpp
@@ -0,0 +1,522 @@
+// Copyright Epic Games, Inc. All Rights Reserved.
+
+#include "Core/PBIKSolver.h"
+#include "Core/PBIKBody.h"
+#include "Core/PBIKConstraint.h"
+#include "Core/PBIKDebug.h"
+
+//#pragma optimize("", off)
+
+namespace PBIK
+{
+	FEffector::FEffector(FBone* InBone)
+	{
+		check(InBone);
+		Bone = InBone;
+		SetGoal(Bone->Position, Bone->Rotation, 1.0f);
+	}
+
+	void FEffector::SetGoal(
+		const FVector InPositionGoal,
+		const FQuat& InRotationGoal,
+		float InAlpha)
+	{
+		PositionOrig = Bone->Position;
+		RotationOrig = Bone->Rotation;
+
+		Position = PositionGoal = InPositionGoal;
+		Rotation = RotationGoal = InRotationGoal;
+
+		Alpha = InAlpha;
+	}
+
+	void FEffector::UpdateFromInputs()
+	{
+		Position = FMath::Lerp(PositionOrig, PositionGoal, Alpha);
+		Rotation = FMath::Lerp(RotationOrig, RotationGoal, Alpha);
+		Pin.Pin()->GoalPoint = Position;
+		Pin.Pin()->Alpha = Alpha;
+	}
+
+	void FEffector::SquashSubRoots()
+	{
+		// optionally apply a preferred angle to give solver a hint which direction to favor
+		// apply amount of preferred angle proportional to the amount this sub-limb is squashed
+		if (!ParentSubRoot || DistToSubRootOrig <= SMALL_NUMBER)
+		{
+			return;
+		}
+
+		float DistToNearestSubRoot = (ParentSubRoot->Position - Position).Size();
+		if (DistToNearestSubRoot >= DistToSubRootOrig)
+		{
+			return; // limb is stretched
+		}
+
+		// shrink distance to reach full blend to preferred angle
+		float ScaledDistOrig = DistToSubRootOrig * 0.3f;
+		// amount squashed (clamped to scaled original length)
+		float DeltaSquash = DistToSubRootOrig - DistToNearestSubRoot;
+		DeltaSquash = DeltaSquash > ScaledDistOrig ? ScaledDistOrig : DeltaSquash;
+		float SquashPercent = DeltaSquash / ScaledDistOrig;
+		if (SquashPercent < 0.01f)
+		{
+			return; // limb not squashed enough
+		}
+
+		FBone* Parent = Bone->Parent;
+		while (Parent)
+		{
+			if (Parent->Body->J.bUsePreferredAngles)
+			{
+				FRotator Clamped = Parent->Body->J.PreferredAngles;
+				FQuat PartialRotation = FQuat::FastLerp(FQuat::Identity, FQuat(Parent->Body->J.PreferredAngles), SquashPercent);
+				Parent->Body->Rotation = Parent->Body->Rotation * PartialRotation;
+				Parent->Body->Rotation.Normalize();
+			}
+
+			if (Parent == ParentSubRoot)
+			{
+				return;
+			}
+
+			Parent = Parent->Parent;
+		}
+	}
+
+} // namespace
+
+void FPBIKSolver::Solve(const FPBIKSolverSettings& Settings)
+{
+	SCOPE_CYCLE_COUNTER(STAT_PBIK_Solve);
+
+	using PBIK::FEffector;
+	using PBIK::FRigidBody;
+	using PBIK::FBone;
+	using PBIK::FBoneSettings;
+
+	// don't run until properly initialized
+	if (!Initialize())
+	{
+		return;
+	}
+
+	// update Bodies with new bone positions from incoming pose and solver settings
+	for (FRigidBody& Body : Bodies)
+	{
+		Body.UpdateFromInputs(Settings);
+	}
+
+	// optionally pin root in-place
+	RootPin.Pin()->bEnabled = Settings.bPinRoot;
+
+	// blend effectors by Alpha and update pin goals
+	for (FEffector& Effector : Effectors)
+	{
+		Effector.UpdateFromInputs();
+		Effector.SquashSubRoots();
+	}
+
+	// run constraint iterations while allowing stretch, just to get reaching pose
+	for (int32  I = 0; I < Settings.Iterations; ++I)
+	{
+		bool bMoveSubRoots = true;
+		for (auto Constraint : Constraints)
+		{
+			Constraint->Solve(bMoveSubRoots);
+		}
+	}
+
+	if (!Settings.bAllowStretch)
+	{
+		for (FEffector& Effector : Effectors)
+		{
+			Effector.SquashSubRoots(); // update squashing once again
+		}
+
+		for (int32  I = 0; I < Settings.Iterations; ++I)
+		{
+			bool bMoveSubRoots = false;
+			for (auto Constraint : Constraints)
+			{
+				Constraint->Solve(bMoveSubRoots);
+			}
+		}
+
+		for (int32  C = Constraints.Num() - 1; C >= 0; --C)
+		{
+			Constraints[C]->RemoveStretch();
+		}
+	}
+
+	// update Bone transforms controlled by Bodies
+	for (FRigidBody& Body : Bodies)
+	{
+		Body.Bone->Position = Body.Position + Body.Rotation * Body.BoneLocalPosition;
+		Body.Bone->Rotation = Body.Rotation;
+	}
+
+	// update Bone transforms controlled by effectors
+	for (const FEffector& Effector : Effectors)
+	{
+		FBone* Bone = Effector.Bone;
+		Bone->Position = Bone->Parent->Position + Bone->Parent->Rotation * Bone->LocalPositionOrig;
+		Bone->Rotation = Effector.Rotation;
+	}
+
+	// propagate to non-solved bones (requires storage in root to tip order)
+	for (FBone& Bone : Bones)
+	{
+		if (Bone.bIsSolved || !Bone.Parent)
+		{
+			continue;
+		}
+
+		Bone.Position = Bone.Parent->Position + Bone.Parent->Rotation * Bone.LocalPositionOrig;
+		Bone.Rotation = Bone.Parent->Rotation * Bone.LocalRotationOrig;
+	}
+}
+
+bool FPBIKSolver::Initialize()
+{
+	if (bReadyToSimulate)
+	{
+		return true;
+	}
+
+	bReadyToSimulate = false;
+
+	if (!InitBones())
+	{
+		return false;
+	}
+
+	if (!InitBodies())
+	{
+		return false;
+	}
+
+	InitConstraints();
+
+	bReadyToSimulate = true;
+
+	return true;
+}
+
+bool FPBIKSolver::InitBones()
+{
+	using PBIK::FEffector;
+	using PBIK::FBone;
+
+	if (!ensureMsgf(Bones.Num() > 0, TEXT("PBIK: no bones added to solver. Cannot initialize.")))
+	{
+		return false;
+	}
+
+	if (!ensureMsgf(Effectors.Num() > 0, TEXT("PBIK: no effectors added to solver. Cannot initialize.")))
+	{
+		return false;
+	}
+
+	// record solver root pointer
+	int32  NumSolverRoots = 0;
+	for (FBone& Bone : Bones)
+	{
+		if (Bone.bIsSolverRoot)
+		{
+			SolverRoot = &Bone;
+			++NumSolverRoots;
+		}
+	}
+
+	if (!ensureMsgf(SolverRoot, TEXT("PBIK: root bone not set. Cannot initialize.")))
+	{
+		return false;
+	}
+
+	if (!ensureMsgf(NumSolverRoots == 1, TEXT("PBIK: more than 1 bone was marked as solver root. Cannot initialize.")))
+	{
+		return false;
+	}
+
+	// initialize pointers to parents
+	for (FBone& Bone : Bones)
+	{
+		// no parent on root, remains null
+		if (Bone.ParentIndex == -1)
+		{
+			continue;
+		}
+
+		// validate parent
+		bool bIndexInRange = Bone.ParentIndex >= 0 && Bone.ParentIndex < Bones.Num();
+		if (!ensureMsgf(bIndexInRange,
+			TEXT("PBIK: bone found with invalid parent index. Cannot initialize.")))
+		{
+			return false;
+		}
+
+		// record parent
+		Bone.Parent = &Bones[Bone.ParentIndex];
+	}
+
+	// initialize IsSolved state by walking up from each effector to the root
+	for (FEffector& Effector : Effectors)
+	{
+		FBone* NextBone = Effector.Bone;
+		while (NextBone)
+		{
+			NextBone->bIsSolved = true;
+			NextBone = NextBone->Parent;
+			if (NextBone && NextBone->bIsSolverRoot)
+			{
+				NextBone->bIsSolved = true;
+				break;
+			}
+		}
+	}
+
+	// initialize local bone transforms
+	for (FBone& Bone : Bones)
+	{
+		FBone* Parent = Bone.Parent;
+		if (!Parent)
+		{
+			continue;
+		}
+
+		Bone.LocalPositionOrig = Parent->Rotation.Inverse() * (Bone.Position - Parent->Position);
+		Bone.LocalRotationOrig = Parent->Rotation.Inverse() * Bone.Rotation;
+	}
+
+	// initialize children lists
+	for (FBone& Parent : Bones)
+	{
+		for (FBone& Child : Bones)
+		{
+			if (Child.bIsSolved && Child.Parent == &Parent)
+			{
+				Parent.Children.Add(&Child);
+			}
+		}
+	}
+
+	// initialize IsSubRoot flag
+	for (FBone& Bone : Bones)
+	{
+		Bone.bIsSubRoot = Bone.Children.Num() > 1 || Bone.bIsSolverRoot;
+	}
+
+	// initialize Effector's nearest ParentSubRoot (FBone) pointer
+	// must be done AFTER setting: Bone.IsSubRoot/IsSolverRoot/Parent
+	for (FEffector& Effector : Effectors)
+	{
+		FBone* Parent = Effector.Bone->Parent;
+		while (Parent)
+		{
+			if (Parent->bIsSubRoot || Parent->bIsSolverRoot)
+			{
+				Effector.ParentSubRoot = Parent;
+				Effector.DistToSubRootOrig = (Parent->Position - Effector.Bone->Position).Size();
+				break;
+			}
+
+			Parent = Parent->Parent;
+		}
+	}
+
+	return true;
+}
+
+bool FPBIKSolver::InitBodies()
+{
+	using PBIK::FEffector;
+	using PBIK::FRigidBody;
+	using PBIK::FBone;
+
+	Bodies.Empty();
+
+	// create bodies
+	for (FEffector& Effector : Effectors)
+	{
+		FBone* NextBone = Effector.Bone;
+		while (true)
+		{
+			FBone* BodyBone = NextBone->Parent;
+			if (!ensureMsgf(BodyBone, TEXT("PBIK: effector is on bone that is not child of root bone.")))
+			{
+				return false;
+			}
+
+			AddBodyForBone(BodyBone);
+
+			NextBone = BodyBone;
+			if (NextBone == SolverRoot)
+			{
+				break;
+			}
+		}
+	}
+
+	// initialize bodies
+	for (FRigidBody& Body : Bodies)
+	{
+		Body.Initialize(SolverRoot);
+	}
+
+	// sort bodies root to leaf
+	Bodies.Sort();
+	Algo::Reverse(Bodies);
+
+	// store pointers to bodies on bones (after sort!)
+	for (FRigidBody& Body : Bodies)
+	{
+		Body.Bone->Body = &Body;
+	}
+
+	return true;
+}
+
+void FPBIKSolver::AddBodyForBone(PBIK::FBone* Bone)
+{
+	for (PBIK::FRigidBody& Body : Bodies)
+	{
+		if (Body.Bone == Bone)
+		{
+			return; // no duplicates
+		}
+	}
+	Bodies.Emplace(Bone);
+}
+
+void FPBIKSolver::InitConstraints()
+{
+	using PBIK::FEffector;
+	using PBIK::FRigidBody;
+	using PBIK::FPinConstraint;
+	using PBIK::FJointConstraint;
+
+	Constraints.Empty();
+
+	// pin root body to animated location (usually disabled by solver settings)
+	TSharedPtr<FPinConstraint> RootConstraint = MakeShared<FPinConstraint>(SolverRoot->Body, SolverRoot->Position);
+	Constraints.Add(RootConstraint);
+	RootPin = RootConstraint;
+
+	// pin bodies to effectors
+	for (FEffector& Effector : Effectors)
+	{
+		if (!ensureMsgf(Effector.Bone->Parent, TEXT("PBIK: effector is on bone that does not have a parent.")))
+		{
+			return;
+		}
+
+		FRigidBody* Body = Effector.Bone->Parent->Body;
+		TSharedPtr<FPinConstraint> Constraint = MakeShared<FPinConstraint>(Body, Effector.Position);
+		Effector.Pin = Constraint;
+		Body->bPinnedToEffector = true;
+		Constraints.Add(Constraint);
+	}
+
+	// constrain all bodies together (child to parent)
+	for (FRigidBody& Body : Bodies)
+	{
+		FRigidBody* ParentBody = Body.GetParentBody();
+		if (!ParentBody)
+		{
+			continue; // root
+		}
+
+		TSharedPtr<FJointConstraint> Constraint = MakeShared<FJointConstraint>(ParentBody, &Body);
+		Constraints.Add(Constraint);
+	}
+}
+
+PBIK::FDebugDraw* FPBIKSolver::GetDebugDraw()
+{
+	DebugDraw.Solver = this;
+	return &DebugDraw;
+}
+
+void FPBIKSolver::Reset()
+{
+	bReadyToSimulate = false;
+	SolverRoot = nullptr;
+	RootPin = nullptr;
+	Bodies.Empty();
+	Bones.Empty();
+	Constraints.Empty();
+	Effectors.Empty();
+}
+
+bool FPBIKSolver::IsReadyToSimulate()
+{
+	return bReadyToSimulate;
+}
+
+int32 FPBIKSolver::AddBone(
+	const FName Name,
+	const int32 ParentIndex,
+	const FVector& InOrigPosition,
+	const FQuat& InOrigRotation,
+	bool bIsSolverRoot)
+{
+	return Bones.Emplace(Name, ParentIndex, InOrigPosition, InOrigRotation, bIsSolverRoot);
+}
+
+bool FPBIKSolver::AddEffector(FName BoneName)
+{
+	for (PBIK::FBone& Bone : Bones)
+	{
+		if (Bone.Name == BoneName)
+		{
+			Effectors.Emplace(&Bone);
+			return true;
+		}
+	}
+
+	return false;
+}
+
+void FPBIKSolver::SetBoneTransform(
+	const int32 Index,
+	const FTransform& InTransform)
+{
+	check(Index >= 0 && Index < Bones.Num());
+	Bones[Index].Position = InTransform.GetLocation();
+	Bones[Index].Rotation = InTransform.GetRotation();
+}
+
+PBIK::FBoneSettings* FPBIKSolver::GetBoneSettings(const int32 Index)
+{
+	// make sure to call Initialize() before applying bone settings
+	if (!ensureMsgf(bReadyToSimulate, TEXT("PBIK: trying to access Bone Settings before Solver is initialized.")))
+	{
+		return nullptr;
+	}
+
+	if (!ensureMsgf(Bones.IsValidIndex(Index), TEXT("PBIK: trying to access Bone Settings with invalid bone index.")))
+	{
+		return nullptr;
+	}
+
+	return &Bones[Index].Body->J;
+}
+
+void FPBIKSolver::GetBoneGlobalTransform(const int32 Index, FTransform& OutTransform)
+{
+	check(Index >= 0 && Index < Bones.Num());
+	PBIK::FBone& Bone = Bones[Index];
+	OutTransform.SetLocation(Bone.Position);
+	OutTransform.SetRotation(Bone.Rotation);
+}
+
+void FPBIKSolver::SetEffectorGoal(
+	const int32 Index,
+	const FVector& InPosition,
+	const FQuat& InRotation,
+	const float Alpha)
+{
+	check(Index >= 0 && Index < Effectors.Num());
+	Effectors[Index].SetGoal(InPosition, InRotation, Alpha);
+}
diff --git a/Engine/Plugins/Experimental/Animation/PBIK/Source/PBIK/Private/PBIK.cpp b/Engine/Plugins/Experimental/Animation/PBIK/Source/PBIK/Private/PBIK.cpp
new file mode 100644
--- /dev/null
+++ b/Engine/Plugins/Experimental/Animation/PBIK/Source/PBIK/Private/PBIK.cpp
@@ -0,0 +1,22 @@
+// Copyright Epic Games, Inc. All Rights Reserved.
+
+#include "PBIK.h"
+
+#define LOCTEXT_NAMESPACE "FPBIKModule"
+
+DEFINE_LOG_CATEGORY(LogIKSolver);
+
+void FPBIKModule::StartupModule()
+{
+	// This code will execute after your module is loaded into memory; the exact timing is specified in the .uplugin file per-module
+}
+
+void FPBIKModule::ShutdownModule()
+{
+	// This function may be called during shutdown to clean up your module.  For modules that support dynamic reloading,
+	// we call this function before unloading the module.
+}
+
+#undef LOCTEXT_NAMESPACE
+	
+IMPLEMENT_MODULE(FPBIKModule, PBIK)
\ No newline at end of file
diff --git a/Engine/Plugins/Experimental/Animation/PBIK/Source/PBIK/Private/RigUnit_PBIK.cpp b/Engine/Plugins/Experimental/Animation/PBIK/Source/PBIK/Private/RigUnit_PBIK.cpp
new file mode 100644
--- /dev/null
+++ b/Engine/Plugins/Experimental/Animation/PBIK/Source/PBIK/Private/RigUnit_PBIK.cpp
@@ -0,0 +1,112 @@
+// Copyright Epic Games, Inc. All Rights Reserved.
+
+#include "RigUnit_PBIK.h"
+#include "Units/RigUnitContext.h"
+
+FRigUnit_PBIK_Execute()
+{
+	DECLARE_SCOPE_HIERARCHICAL_COUNTER_RIGUNIT()
+
+	FRigHierarchyContainer* Hierarchy = ExecuteContext.Hierarchy;
+	if (Hierarchy == nullptr)
+	{
+		return;
+	}
+
+	FRigBoneHierarchy& Bones = Hierarchy->BoneHierarchy;
+
+	if (Context.State == EControlRigState::Init)
+	{
+		// check how many effectors are assigned to a bone
+		int NumEffectors = 0;
+		for (const FPBIKEffector& Effector : Effectors)
+		{
+			if (Bones.GetIndex(Effector.Bone) != NAME_None)
+			{
+				++NumEffectors; // bone is set and exists!
+			}
+		}
+
+		// validate inputs are ready to be initialized
+		bool bHasEffectors = NumEffectors > 0;
+		bool bRootIsAssigned = Root != NAME_None;
+		if (!(bHasEffectors && bRootIsAssigned))
+		{
+			return; // not setup yet
+		}
+
+		// reset all internal data
+		Solver.Reset();
+
+		// create bones
+		for (int B = 0; B < Bones.Num(); ++B)
+		{
+			FName Name = Bones[B].Name;
+			int ParentIndex = Bones[B].ParentIndex;
+			FVector InOrigPosition = Bones[B].InitialTransform.GetLocation();
+			FQuat InOrigRotation = Bones[B].InitialTransform.GetRotation();
+			bool bIsRoot = Bones[B].Name == Root;
+			Solver.AddBone(Name, ParentIndex, InOrigPosition, InOrigRotation, bIsRoot);
+		}
+		
+		// create effectors
+		for (const FPBIKEffector& Effector : Effectors)
+		{
+			if (!Solver.AddEffector(Effector.Bone))
+			{
+				return;
+			}
+		}
+		
+		// initialize
+		Solver.Initialize();
+
+		return; // don't update during init
+	}
+
+	if (!Solver.IsReadyToSimulate())
+	{
+		return;
+	}
+
+	// set bones to input pose
+	for (int B = 0; B < Bones.Num(); ++B)
+	{
+		Solver.SetBoneTransform(B, Bones[B].GlobalTransform);
+	}
+
+	// update bone settings
+	for (const FPBIKBoneSetting& BoneSetting : BoneSettings)
+	{
+		int32 BoneIndex = Bones.GetIndex(BoneSetting.Bone);
+		if (BoneIndex == INDEX_NONE)
+		{
+			continue; // no bones to apply it to
+		}
+
+		if (PBIK::FBoneSettings* InternalSettings = Solver.GetBoneSettings(BoneIndex))
+		{
+			BoneSetting.CopyToCoreStruct(*InternalSettings);
+		}
+	}
+
+	// update effectors
+	for (int E = 0; E < Effectors.Num(); ++E)
+	{
+		FVector Position = Effectors[E].Transform.GetLocation();
+		FQuat Rotation = Effectors[E].Transform.GetRotation();
+		Solver.SetEffectorGoal(E, Position, Rotation, Effectors[E].Alpha);
+	}
+
+	// solve
+	Solver.Solve(Settings);
+
+	// copy transforms back
+	for (int B = 0; B < Bones.Num(); ++B)
+	{
+		Solver.GetBoneGlobalTransform(B, Bones[B].GlobalTransform);
+	}
+
+	// do all debug drawing
+	Debug.Draw(Context.DrawInterface, &Solver);
+}
diff --git a/Engine/Plugins/Experimental/Animation/PBIK/Source/PBIK/Public/Core/PBIKBody.h b/Engine/Plugins/Experimental/Animation/PBIK/Source/PBIK/Public/Core/PBIKBody.h
new file mode 100644
--- /dev/null
+++ b/Engine/Plugins/Experimental/Animation/PBIK/Source/PBIK/Public/Core/PBIKBody.h
@@ -0,0 +1,111 @@
+// Copyright Epic Games, Inc. All Rights Reserved.
+
+#pragma once
+
+struct FPBIKSolverSettings;
+
+namespace PBIK
+{
+
+struct FJointConstraint;
+struct FRigidBody;
+
+struct FBone
+{
+	FName Name;
+	int ParentIndex = -2; // -2 is unset, -1 for root, or 0...n otherwise
+	bool bIsSolverRoot = false;
+	bool bIsSolved = false;
+	bool bIsSubRoot = false;
+	FVector Position;
+	FQuat Rotation;
+	FVector LocalPositionOrig;
+	FQuat LocalRotationOrig;
+
+	// initialized - these fields are null/empty until after Solver::Initialize()
+	FRigidBody* Body = nullptr;
+	FBone* Parent = nullptr;
+	TArray<FBone*> Children;
+	// initialized
+
+	FBone(
+		const FName InName,
+		const int& InParentIndex,		// must pass -1 for root of whole skeleton
+		const FVector& InOrigPosition,
+		const FQuat& InOrigRotation,
+		bool bInIsSolverRoot);
+
+	bool HasChild(const FBone* Bone);
+};
+
+enum class ELimitType : uint8
+{
+	Free,
+	Limited,
+	Locked,
+};
+
+struct FBoneSettings
+{
+	float RotationStiffness = 0.0f; // range (0, 1)
+	float PositionStiffness = 0.0f; // range (0, 1)
+
+	ELimitType X;
+	float MinX = 0.0f; // range (-180, 180)
+	float MaxX = 0.0f;
+
+	ELimitType Y;
+	float MinY = 0.0f;
+	float MaxY = 0.0f;
+
+	ELimitType Z;
+	float MinZ = 0.0f;
+	float MaxZ = 0.0f;
+
+	bool bUsePreferredAngles = false;
+	FRotator PreferredAngles = FRotator::ZeroRotator;
+};
+
+struct FRigidBody
+{
+	FBone* Bone = nullptr;
+	FBoneSettings J;
+
+	FVector Position;
+	FQuat Rotation;
+	FQuat RotationOrig;
+	FVector BoneLocalPosition;
+	TArray<FVector> ChildLocalPositions;
+
+	float InvMass = 0.0f;
+	bool bPinnedToEffector = false;
+	float Length;
+	
+private:
+
+	int NumBonesToRoot = 0;
+
+public:
+
+	FRigidBody(FBone* InBone);
+
+	void Initialize(FBone* SolverRoot);
+
+	void UpdateFromInputs(const FPBIKSolverSettings& Settings);
+
+	int GetNumBonesToRoot() const;
+
+	FRigidBody* GetParentBody();
+
+	void ApplyPushToRotateBody(const FVector& Push, const FVector& Offset);
+	
+	void ApplyPushToPosition(const FVector& Push);
+};
+
+// for sorting Bodies hierarchically (root to leaf order)
+inline bool operator<(const FRigidBody& Lhs, const FRigidBody& Rhs)
+{ 
+	return Lhs.GetNumBonesToRoot() < Rhs.GetNumBonesToRoot(); 
+}
+
+} // namespace
\ No newline at end of file
diff --git a/Engine/Plugins/Experimental/Animation/PBIK/Source/PBIK/Public/Core/PBIKConstraint.h b/Engine/Plugins/Experimental/Animation/PBIK/Source/PBIK/Public/Core/PBIKConstraint.h
new file mode 100644
--- /dev/null
+++ b/Engine/Plugins/Experimental/Animation/PBIK/Source/PBIK/Public/Core/PBIKConstraint.h
@@ -0,0 +1,104 @@
+// Copyright Epic Games, Inc. All Rights Reserved.
+
+#pragma once
+
+
+namespace PBIK
+{
+
+struct FConstraint
+{
+	bool bEnabled = true;
+
+	virtual void Solve(bool bMoveSubRoots) = 0;
+	virtual void RemoveStretch() {};
+};
+
+struct FJointConstraint : public FConstraint
+{
+
+private:
+
+	FRigidBody* A;
+	FRigidBody* B;
+
+	FVector PinPointLocalToA;
+	FVector PinPointLocalToB;
+
+	FVector XOrig;
+	FVector YOrig;
+	FVector ZOrig;
+
+	FVector XA;
+	FVector YA;
+	FVector ZA;
+	FVector XB;
+	FVector YB;
+	FVector ZB;
+
+	FVector ZBProjOnX;
+	FVector ZBProjOnY;
+	FVector YBProjOnZ;
+
+	float AngleX;
+	float AngleY;
+	float AngleZ;
+
+public:
+
+	FJointConstraint(FRigidBody* InA, FRigidBody* InB);
+
+	virtual ~FJointConstraint() {};
+
+	virtual void Solve(bool bMoveSubRoots) override;
+
+	virtual void RemoveStretch() override;
+
+private:
+
+	FVector GetPositionCorrection(FVector& OutBodyToA, FVector& OutBodyToB);
+
+	void ApplyRotationCorrection(FQuat PureRotA, FQuat PureRotB);
+
+	void UpdateJointLimits();
+
+	void RotateWithinLimits(
+		float MinAngle,
+		float MaxAngle,
+		float CurrentAngle,
+		FVector RotAxis,
+		FVector CurVec,
+		FVector RefVec);
+
+	void UpdateLocalRotateAxes(bool bX, bool bY, bool bZ);
+
+	void DecomposeRotationAngles();
+
+	float SignedAngleBetweenNormals(
+		const FVector& From,
+		const FVector& To,
+		const FVector& Axis) const;
+};
+
+struct FPinConstraint : public FConstraint
+{
+	FVector GoalPoint;
+	FRigidBody* A;
+	FVector PinPointLocalToA;
+	float Alpha = 1.0;
+
+public:
+
+	FPinConstraint(FRigidBody* InBody, const FVector& InPinPoint);
+
+	virtual ~FPinConstraint() {};
+
+	virtual void Solve(bool bMoveSubRoots) override;
+
+private:
+
+	FVector GetPositionCorrection(FVector& OutBodyToPinPoint) const;
+};
+
+} // namespace
+
diff --git a/Engine/Plugins/Experimental/Animation/PBIK/Source/PBIK/Public/Core/PBIKDebug.h b/Engine/Plugins/Experimental/Animation/PBIK/Source/PBIK/Public/Core/PBIKDebug.h
new file mode 100644
--- /dev/null
+++ b/Engine/Plugins/Experimental/Animation/PBIK/Source/PBIK/Public/Core/PBIKDebug.h
@@ -0,0 +1,27 @@
+// Copyright Epic Games, Inc. All Rights Reserved.
+
+#pragma once
+
+struct FPBIKSolver;
+
+namespace PBIK
+{
+	struct FDebugLine
+	{
+		FVector A;
+		FVector B;
+
+		FDebugLine(const FVector& InA, const FVector& InB)
+		{
+			A = InA;
+			B = InB;
+		}
+	};
+
+	struct FDebugDraw
+	{
+		FPBIKSolver* Solver;
+
+		void GetDebugLinesForBodies(TArray<FDebugLine>& OutLines);
+	};
+}
diff --git a/Engine/Plugins/Experimental/Animation/PBIK/Source/PBIK/Public/Core/PBIKSolver.h b/Engine/Plugins/Experimental/Animation/PBIK/Source/PBIK/Public/Core/PBIKSolver.h
new file mode 100644
--- /dev/null
+++ b/Engine/Plugins/Experimental/Animation/PBIK/Source/PBIK/Public/Core/PBIKSolver.h
@@ -0,0 +1,135 @@
+// Copyright Epic Games, Inc. All Rights Reserved.
+
+#pragma once
+
+#include "PBIKBody.h"
+#include "PBIKConstraint.h"
+#include "PBIKDebug.h"
+
+#include "PBIKSolver.generated.h"
+
+DECLARE_CYCLE_STAT(TEXT("PBIK Solve"), STAT_PBIK_Solve, STATGROUP_Anim);
+
+namespace PBIK
+{
+	static float GLOBAL_UNITS = 100.0f; // (1.0f = meters), (100.0f = centimeters)
+
+struct FEffector
+{
+	FVector Position;
+	FQuat Rotation;
+
+	FVector PositionOrig;
+	FQuat RotationOrig;
+
+	FVector PositionGoal;
+	FQuat RotationGoal;
+
+	FBone* Bone;
+	TWeakPtr<FPinConstraint> Pin;
+	
+	float DistToSubRootOrig;
+	FBone* ParentSubRoot;
+
+	float Alpha;
+
+	FEffector(FBone* InBone);
+
+	void SetGoal( const FVector InPositionGoal, const FQuat& InRotationGoal, float InAlpha);
+
+	void UpdateFromInputs();
+
+	void SquashSubRoots();
+};
+
+} // namespace
+
+USTRUCT()
+struct FPBIKSolverSettings
+{
+	GENERATED_BODY()
+
+	UPROPERTY(meta = (ClampMin = "0", UIMin = "0.0", UIMax = "200.0"))
+	int32 Iterations = 20;
+
+	UPROPERTY(meta = (ClampMin = "0", UIMin = "0.0", UIMax = "10.0"))
+	float MassMultiplier = 1.0f;
+
+	UPROPERTY()
+	bool bAllowStretch = false;
+
+	UPROPERTY()
+	bool bPinRoot = false;
+};
+
+
+
+USTRUCT()
+struct FPBIKSolver
+{
+	GENERATED_BODY()
+
+public:
+
+	PBIK::FDebugDraw* GetDebugDraw();
+
+	//
+	// main runtime functions
+	//
+
+	bool Initialize();
+
+	void Solve(const FPBIKSolverSettings& Settings);
+
+	void Reset();
+
+	bool IsReadyToSimulate();
+
+	//
+	// set input / get output at runtime
+	//
+
+	void SetBoneTransform(const int32 Index, const FTransform& InTransform);
+
+	PBIK::FBoneSettings* GetBoneSettings(const int32 Index);
+
+	void SetEffectorGoal(const int32 Index, const FVector& InPosition, const FQuat& InRotation, const float Alpha);
+
+	void GetBoneGlobalTransform(const int32 Index, FTransform& OutTransform);
+
+	//
+	// pre-init /  setup functions
+	//
+
+	int32 AddBone(
+		const FName Name,
+		const int32 ParentIndex,
+		const FVector& InOrigPosition,
+		const FQuat& InOrigRotation,
+		bool bIsSolverRoot);
+
+	bool AddEffector(FName BoneName);
+	
+private:
+
+	bool InitBones();
+
+	bool InitBodies();
+
+	void InitConstraints();
+
+	void AddBodyForBone(PBIK::FBone* Bone);
+
+private:
+
+	PBIK::FBone* SolverRoot = nullptr;
+	TWeakPtr<PBIK::FPinConstraint> RootPin = nullptr;
+	TArray<PBIK::FBone> Bones;
+	TArray<PBIK::FRigidBody> Bodies;
+	TArray<TSharedPtr<PBIK::FConstraint>> Constraints;
+	TArray<PBIK::FEffector> Effectors;
+	bool bReadyToSimulate = false;
+	
+	PBIK::FDebugDraw DebugDraw;
+	friend PBIK::FDebugDraw;
+};
diff --git a/Engine/Plugins/Experimental/Animation/PBIK/Source/PBIK/Public/PBIK.h b/Engine/Plugins/Experimental/Animation/PBIK/Source/PBIK/Public/PBIK.h
new file mode 100644
--- /dev/null
+++ b/Engine/Plugins/Experimental/Animation/PBIK/Source/PBIK/Public/PBIK.h
@@ -0,0 +1,17 @@
+// Copyright Epic Games, Inc. All Rights Reserved.
+
+#pragma once
+
+#include "CoreMinimal.h"
+#include "Modules/ModuleManager.h"
+
+class FPBIKModule : public IModuleInterface
+{
+public:
+
+	/** IModuleInterface implementation */
+	virtual void StartupModule() override;
+	virtual void ShutdownModule() override;
+};
+
+DECLARE_LOG_CATEGORY_EXTERN(LogIKSolver, Warning, All);
diff --git a/Engine/Plugins/Experimental/Animation/PBIK/Source/PBIK/Public/RigUnit_PBIK.h b/Engine/Plugins/Experimental/Animation/PBIK/Source/PBIK/Public/RigUnit_PBIK.h
new file mode 100644
--- /dev/null
+++ b/Engine/Plugins/Experimental/Animation/PBIK/Source/PBIK/Public/RigUnit_PBIK.h
@@ -0,0 +1,160 @@
+// Copyright Epic Games, Inc. All Rights Reserved.
+
+#pragma once
+
+#include "Units/Highlevel/RigUnit_HighlevelBase.h"
+#include "Drawing/ControlRigDrawInterface.h"
+
+#include "Core/PBIKSolver.h"
+#include "Core/PBIKDebug.h"
+
+#include "RigUnit_PBIK.generated.h"
+
+using PBIK::FDebugLine;
+
+USTRUCT()
+struct FPBIKDebug
+{
+	GENERATED_BODY()
+
+	UPROPERTY()
+	float DrawScale = 1.0f;
+
+	UPROPERTY()
+	bool bDrawDebug = false;
+
+	void Draw(FControlRigDrawInterface* DrawInterface, FPBIKSolver* Solver) const
+	{
+		if (!(DrawInterface && Solver && bDrawDebug))
+		{
+			return;
+		}
+
+		const FLinearColor Bright = FLinearColor(0.f, 1.f, 1.f, 1.f);
+		DrawInterface->DrawBox(FTransform::Identity, FTransform(FQuat::Identity, FVector(0, 0, 0), FVector(1.f, 1.f, 1.f) * DrawScale * 0.1f), Bright);
+
+		TArray<FDebugLine> BodyLines;
+		Solver->GetDebugDraw()->GetDebugLinesForBodies(BodyLines);
+		const FLinearColor BodyColor = FLinearColor(0.1f, 0.1f, 1.f, 1.f);
+		for (FDebugLine Line : BodyLines)
+		{
+			DrawInterface->DrawLine(FTransform::Identity, Line.A, Line.B, BodyColor);
+		}
+	}
+};
+
+USTRUCT(BlueprintType)
+struct FPBIKEffector
+{
+	GENERATED_BODY()
+
+	FPBIKEffector()	: Bone(NAME_None) {}
+
+	UPROPERTY(meta = (Constant, CustomWidget = "BoneName"))
+	FName Bone;
+
+	UPROPERTY()
+	FTransform Transform;
+
+	UPROPERTY()
+	float Alpha = 1.0f;
+};
+
+UENUM(BlueprintType)
+enum class EPBIKLimitType : uint8
+{
+	Free,
+	Limited,
+	Locked,
+};
+
+USTRUCT(BlueprintType)
+struct FPBIKBoneSetting
+{
+	GENERATED_BODY()
+
+	FPBIKBoneSetting() : Bone(NAME_None) {}
+
+	UPROPERTY(meta = (Constant, CustomWidget = "BoneName"))
+	FName Bone;
+
+	UPROPERTY(meta = (ClampMin = "0", ClampMax = "1", UIMin = "0.0", UIMax = "1.0"))
+	float RotationStiffness = 0.0f;
+	UPROPERTY(meta = (ClampMin = "0", ClampMax = "1", UIMin = "0.0", UIMax = "1.0"))
+	float PositionStiffness = 0.0f;
+
+	UPROPERTY()
+	EPBIKLimitType X;
+	UPROPERTY(meta = (ClampMin = "-180", ClampMax = "0", UIMin = "-180.0", UIMax = "0.0"))
+	float MinX = 0.0f;
+	UPROPERTY(meta = (ClampMin = "0", ClampMax = "180", UIMin = "0.0", UIMax = "180.0"))
+	float MaxX = 0.0f;
+
+	UPROPERTY()
+	EPBIKLimitType Y;
+	UPROPERTY(meta = (ClampMin = "-180", ClampMax = "0", UIMin = "-180.0", UIMax = "0.0"))
+	float MinY = 0.0f;
+	UPROPERTY(meta = (ClampMin = "0", ClampMax = "180", UIMin = "0.0", UIMax = "180.0"))
+	float MaxY = 0.0f;
+
+	UPROPERTY()
+	EPBIKLimitType Z;
+	UPROPERTY(meta = (ClampMin = "-180", ClampMax = "0", UIMin = "-180.0", UIMax = "0.0"))
+	float MinZ = 0.0f;
+	UPROPERTY(meta = (ClampMin = "0", ClampMax = "180", UIMin = "0.0", UIMax = "180.0"))
+	float MaxZ = 0.0f;
+
+	UPROPERTY()
+	bool bUsePreferredAngles = false;
+	UPROPERTY()
+	FRotator PreferredAngles;
+
+	void CopyToCoreStruct(PBIK::FBoneSettings& Settings) const
+	{
+		Settings.RotationStiffness = RotationStiffness;
+		Settings.PositionStiffness = PositionStiffness;
+		Settings.X = static_cast<PBIK::ELimitType>(X);
+		Settings.MinX = MinX;
+		Settings.MaxX = MaxX;
+		Settings.Y = static_cast<PBIK::ELimitType>(Y);
+		Settings.MinY = MinY;
+		Settings.MaxY = MaxY;
+		Settings.Z = static_cast<PBIK::ELimitType>(Z);
+		Settings.MinZ = MinZ;
+		Settings.MaxZ = MaxZ;
+		Settings.bUsePreferredAngles = bUsePreferredAngles;
+		Settings.PreferredAngles = PreferredAngles;
+	}
+};
+
+USTRUCT(meta=(DisplayName="Position Based IK", Category="Hierarchy", Keywords="Position Based, PBIK, IK, Full Body"))
+struct FRigUnit_PBIK : public FRigUnit_HighlevelBaseMutable
+{
+	GENERATED_BODY()
+
+	RIGVM_METHOD()
+	virtual void Execute(const FRigUnitContext& Context) override;
+
+	FRigUnit_PBIK()
+		: Root(NAME_None)
+	{
+	}
+
+	UPROPERTY(meta = (Input, Constant, CustomWidget = "BoneName"))
+	FName Root;
+
+	UPROPERTY(meta = (Input))
+	TArray<FPBIKEffector> Effectors;
+
+	UPROPERTY(meta = (Input))
+	TArray<FPBIKBoneSetting> BoneSettings;
+
+	UPROPERTY(meta = (Input))
+	FPBIKSolverSettings Settings;
+
+	UPROPERTY(meta = (Input))
+	FPBIKDebug Debug;
+
+	UPROPERTY(transient)
+	FPBIKSolver Solver;
+};
