PBIK: fixed crash when trying to apply bone setting to bone that is not solved (between root and effector).
JIRA: UE-107414

diff --git a/Engine/Plugins/Experimental/Animation/PBIK/Source/PBIK/Private/Core/PBIKSolver.cpp b/Engine/Plugins/Experimental/Animation/PBIK/Source/PBIK/Private/Core/PBIKSolver.cpp
--- a/Engine/Plugins/Experimental/Animation/PBIK/Source/PBIK/Private/Core/PBIKSolver.cpp
+++ b/Engine/Plugins/Experimental/Animation/PBIK/Source/PBIK/Private/Core/PBIKSolver.cpp
@@ -500,6 +500,12 @@ PBIK::FBoneSettings* FPBIKSolver::GetBoneSettings(const int32 Index)
 		return nullptr;
 	}
 
+	if (!Bones[Index].Body)
+	{
+		UE_LOG(LogTemp, Warning, TEXT("PBIK: trying to apply Bone Settings to bone that is not simulated (not between root and effector)."));
+		return nullptr;
+	}
+
 	return &Bones[Index].Body->J;
 }
 
