PBIK: misc changes
- fixed bug, crash when effector applied to root bone
- added support for sparse effector indices
- fixed bug, did not correctly detect condition where effector bone was not set
- fixed clashing logging category name
- integrated solver into IKRig plugin
JIRA: UE-107466, UE-107411, UE-108212

diff --git a/Engine/Plugins/Animation/IKRig/Source/IKRig/Public/IKRigSolver.h b/Engine/Plugins/Animation/IKRig/Source/IKRig/Public/IKRigSolver.h
--- a/Engine/Plugins/Animation/IKRig/Source/IKRig/Public/IKRigSolver.h
+++ b/Engine/Plugins/Animation/IKRig/Source/IKRig/Public/IKRigSolver.h
@@ -32,7 +32,7 @@ class IKRIG_API UIKRigSolver : public UObject
 	
 public:
 
-	UPROPERTY(EditAnywhere, Category = "Definition")
+	UPROPERTY(EditAnywhere, Category = "Settings")
 	bool bEnabled = true;
 
 	/** override Init() to setup internal data based on ref pose */
diff --git a/Engine/Plugins/Experimental/Animation/PBIK/PBIK.uplugin b/Engine/Plugins/Experimental/Animation/PBIK/PBIK.uplugin
--- a/Engine/Plugins/Experimental/Animation/PBIK/PBIK.uplugin
+++ b/Engine/Plugins/Experimental/Animation/PBIK/PBIK.uplugin
@@ -19,6 +19,10 @@
 		{
 			"Name": "ControlRig",
 			"Enabled": true
+		},
+		{
+			"Name": "IKRig",
+			"Enabled": true
 		}
 	],
 	"Modules": [
diff --git a/Engine/Plugins/Experimental/Animation/PBIK/Source/PBIK/PBIK.Build.cs b/Engine/Plugins/Experimental/Animation/PBIK/Source/PBIK/PBIK.Build.cs
--- a/Engine/Plugins/Experimental/Animation/PBIK/Source/PBIK/PBIK.Build.cs
+++ b/Engine/Plugins/Experimental/Animation/PBIK/Source/PBIK/PBIK.Build.cs
@@ -13,7 +13,8 @@ public class PBIK : ModuleRules
 			{
 				"Core",
 				"RigVM",
-				"ControlRig"
+				"ControlRig",
+				"IKRig"
 			}
 			);
 
diff --git a/Engine/Plugins/Experimental/Animation/PBIK/Source/PBIK/Private/Core/PBIKBody.cpp b/Engine/Plugins/Experimental/Animation/PBIK/Source/PBIK/Private/Core/PBIKBody.cpp
--- a/Engine/Plugins/Experimental/Animation/PBIK/Source/PBIK/Private/Core/PBIKBody.cpp
+++ b/Engine/Plugins/Experimental/Animation/PBIK/Source/PBIK/Private/Core/PBIKBody.cpp
@@ -1,7 +1,5 @@
 // Copyright Epic Games, Inc. All Rights Reserved.
 
-#pragma once
-
 #include "Core/PBIKBody.h"
 #include "Core/PBIKSolver.h"
 
@@ -101,7 +99,7 @@ FRigidBody* FRigidBody::GetParentBody()
 
 void FRigidBody::ApplyPushToRotateBody(const FVector& Push, const FVector& Offset)
 {
-	// quat differentiation
+	// equation 8 in "Detailed Rigid Body Simulation with XPBD"
 	FVector Omega = InvMass * (1.0f - J.RotationStiffness) * FVector::CrossProduct(Offset, Push);
 	FQuat OQ(Omega.X, Omega.Y, Omega.Z, 0.0f);
 	OQ = OQ * Rotation;
@@ -118,7 +116,7 @@ void FRigidBody::ApplyPushToRotateBody(const FVector& Push, const FVector& Offse
 
 void FRigidBody::ApplyPushToPosition(const FVector& Push)
 {
-	if (bPinnedToEffector)
+	if (AttachedEffector)
 	{
 		return; // pins are locked
 	}
diff --git a/Engine/Plugins/Experimental/Animation/PBIK/Source/PBIK/Private/Core/PBIKConstraint.cpp b/Engine/Plugins/Experimental/Animation/PBIK/Source/PBIK/Private/Core/PBIKConstraint.cpp
--- a/Engine/Plugins/Experimental/Animation/PBIK/Source/PBIK/Private/Core/PBIKConstraint.cpp
+++ b/Engine/Plugins/Experimental/Animation/PBIK/Source/PBIK/Private/Core/PBIKConstraint.cpp
@@ -1,13 +1,13 @@
 // Copyright Epic Games, Inc. All Rights Reserved.
 
-#pragma once
-
 #include "Core/PBIKConstraint.h"
 #include "Core/PBIKSolver.h"
 #include "Math/UnrealMathUtility.h"
 
 namespace PBIK
 {
+
+	static float ALMOST_ONE = 1.0f - KINDA_SMALL_NUMBER;
 	
 FJointConstraint::FJointConstraint(FRigidBody* InA, FRigidBody* InB)
 {
@@ -29,32 +29,41 @@ void FJointConstraint::Solve(bool bMoveSubRoots)
 	FVector OffsetA;
 	FVector OffsetB;
 	FVector Correction = GetPositionCorrection(OffsetA, OffsetB);
+	//Correction = Correction * OneOverInvMassSum;
 
 	// apply rotation from correction 
 	// at pin point to both bodies
-	A->ApplyPushToRotateBody(Correction, OffsetA);
-	B->ApplyPushToRotateBody(-Correction, OffsetB);
+	A->ApplyPushToRotateBody(Correction, OffsetA);//* AInvMass, OffsetA);
+	B->ApplyPushToRotateBody(-Correction, OffsetB); //* BInvMass, OffsetB);
 
 	// enforce joint limits
 	UpdateJointLimits();
 
+	// calc inv mass of body A
+	float AInvMass = 1.0f;
+	AInvMass -= A->AttachedEffector ? A->AttachedEffector->StrengthAlpha : 0.0f;
+	AInvMass -= !bMoveSubRoots && A->Bone->bIsSubRoot ? 1.0f : 0.0f;
+	AInvMass = AInvMass <= 0.0f ? 0.0f : AInvMass;
+
+	// calc inv mass of body B
+	float BInvMass = 1.0f;
+	BInvMass -= B->AttachedEffector ? B->AttachedEffector->StrengthAlpha : 0.0f;
+	BInvMass -= !bMoveSubRoots && B->Bone->bIsSubRoot ? 1.0f : 0.0f;
+	BInvMass = BInvMass <= 0.0f ? 0.0f : BInvMass;
+
+	float InvMassSum = AInvMass + BInvMass;
+	if (FMath::IsNearlyZero(InvMassSum))
+	{
+		return; // no correction can be applied on full locked configuration
+	}
+	float OneOverInvMassSum = 1.0f / InvMassSum;
+
 	// apply positional correction to align pin point on both Bodies
 	// NOTE: applying position correction AFTER rotation takes into
 	// consideration change in relative pin locations mid-step after
 	// bodies have been rotated by ApplyPushToRotateBody
 	Correction = GetPositionCorrection(OffsetA, OffsetB);
-
-	// scale correction by inv mass
-	bool bALocked = A->bPinnedToEffector || (!bMoveSubRoots && A->Bone->bIsSubRoot);
-	bool bBLocked = B->bPinnedToEffector || (!bMoveSubRoots && B->Bone->bIsSubRoot);
-	float AInvMass = bALocked ? 0.0f : 1.0f;
-	float BInvMass = bBLocked ? 0.0f : 1.0f;
-	float InvMassSum = AInvMass + BInvMass;
-	if (InvMassSum < KINDA_SMALL_NUMBER)
-	{
-		return;
-	}
-	Correction = Correction * (1.0f / InvMassSum);
+	Correction = Correction * OneOverInvMassSum;
 
 	A->ApplyPushToPosition(Correction * AInvMass);
 	B->ApplyPushToPosition(-Correction * BInvMass);
diff --git a/Engine/Plugins/Experimental/Animation/PBIK/Source/PBIK/Private/Core/PBIKDebug.cpp b/Engine/Plugins/Experimental/Animation/PBIK/Source/PBIK/Private/Core/PBIKDebug.cpp
--- a/Engine/Plugins/Experimental/Animation/PBIK/Source/PBIK/Private/Core/PBIKDebug.cpp
+++ b/Engine/Plugins/Experimental/Animation/PBIK/Source/PBIK/Private/Core/PBIKDebug.cpp
@@ -1,7 +1,5 @@
 // Copyright Epic Games, Inc. All Rights Reserved.
 
-#pragma once
-
 #include "Core/PBIKDebug.h"
 #include "Core/PBIKSolver.h"
 #include "Core/PBIKBody.h"
diff --git a/Engine/Plugins/Experimental/Animation/PBIK/Source/PBIK/Private/Core/PBIKSolver.cpp b/Engine/Plugins/Experimental/Animation/PBIK/Source/PBIK/Private/Core/PBIKSolver.cpp
--- a/Engine/Plugins/Experimental/Animation/PBIK/Source/PBIK/Private/Core/PBIKSolver.cpp
+++ b/Engine/Plugins/Experimental/Animation/PBIK/Source/PBIK/Private/Core/PBIKSolver.cpp
@@ -13,13 +13,14 @@ namespace PBIK
 	{
 		check(InBone);
 		Bone = InBone;
-		SetGoal(Bone->Position, Bone->Rotation, 1.0f);
+		SetGoal(Bone->Position, Bone->Rotation, 1.0f, 1.0f);
 	}
 
 	void FEffector::SetGoal(
 		const FVector InPositionGoal,
 		const FQuat& InRotationGoal,
-		float InAlpha)
+		float InTransformAlpha,
+		float InStrengthAlpha)
 	{
 		PositionOrig = Bone->Position;
 		RotationOrig = Bone->Rotation;
@@ -27,15 +28,16 @@ namespace PBIK
 		Position = PositionGoal = InPositionGoal;
 		Rotation = RotationGoal = InRotationGoal;
 
-		Alpha = InAlpha;
+		TransformAlpha = InTransformAlpha;
+		StrengthAlpha = InStrengthAlpha;
 	}
 
 	void FEffector::UpdateFromInputs()
 	{
-		Position = FMath::Lerp(PositionOrig, PositionGoal, Alpha);
-		Rotation = FMath::Lerp(RotationOrig, RotationGoal, Alpha);
+		Position = FMath::Lerp(PositionOrig, PositionGoal, TransformAlpha);
+		Rotation = FMath::Lerp(RotationOrig, RotationGoal, TransformAlpha);
 		Pin.Pin()->GoalPoint = Position;
-		Pin.Pin()->Alpha = Alpha;
+		Pin.Pin()->Alpha = StrengthAlpha;
 	}
 
 	void FEffector::SquashSubRoots()
@@ -47,7 +49,12 @@ namespace PBIK
 			return;
 		}
 
-		float DistToNearestSubRoot = (ParentSubRoot->Position - Position).Size();
+		// we have to be careful here when calculating the distance to the parent sub-root.
+		// if the the parent sub-root is attached to an effector, use the effector's position
+		// otherwise use the current position of the FRigidBody
+		FEffector* ParentEffector = ParentSubRoot->AttachedEffector;
+		FVector ParentSubRootPosition = ParentEffector ? ParentEffector->Position : ParentSubRoot->Position;
+		float DistToNearestSubRoot = (ParentSubRootPosition - Position).Size();
 		if (DistToNearestSubRoot >= DistToSubRootOrig)
 		{
 			return; // limb is stretched
@@ -65,7 +72,7 @@ namespace PBIK
 		}
 
 		FBone* Parent = Bone->Parent;
-		while (Parent)
+		while (Parent && Parent->bIsSolved)
 		{
 			if (Parent->Body->J.bUsePreferredAngles)
 			{
@@ -75,7 +82,7 @@ namespace PBIK
 				Parent->Body->Rotation.Normalize();
 			}
 
-			if (Parent == ParentSubRoot)
+			if (Parent == ParentSubRoot->Bone)
 			{
 				return;
 			}
@@ -101,19 +108,42 @@ void FPBIKSolver::Solve(const FPBIKSolverSettings& Settings)
 		return;
 	}
 
+	// initialize local bone transforms
+	// this has to be done every tick because incoming animation can modify these
+	// even the LocalPosition has to be updated incase translation is animated
+	for (FBone& Bone : Bones)
+	{
+		FBone* Parent = Bone.Parent;
+		if (!Parent)
+		{
+			continue;
+		}
+
+		Bone.LocalPositionOrig = Parent->Rotation.Inverse() * (Bone.Position - Parent->Position);
+		Bone.LocalRotationOrig = Parent->Rotation.Inverse() * Bone.Rotation;
+	}
+
 	// update Bodies with new bone positions from incoming pose and solver settings
 	for (FRigidBody& Body : Bodies)
 	{
 		Body.UpdateFromInputs(Settings);
 	}
 
-	// optionally pin root in-place
-	RootPin.Pin()->bEnabled = Settings.bPinRoot;
+	// optionally pin root in-place (convenience, does not require an effector)
+	if (RootPin.IsValid())
+	{
+		RootPin.Pin()->bEnabled = Settings.bPinRoot;
+	}
 
 	// blend effectors by Alpha and update pin goals
 	for (FEffector& Effector : Effectors)
 	{
 		Effector.UpdateFromInputs();
+	}
+
+	// squash sub-roots to apply preferred angles
+	for (FEffector& Effector : Effectors)
+	{
 		Effector.SquashSubRoots();
 	}
 
@@ -129,6 +159,11 @@ void FPBIKSolver::Solve(const FPBIKSolverSettings& Settings)
 
 	if (!Settings.bAllowStretch)
 	{
+		for (int32 C = Constraints.Num() - 1; C >= 0; --C)
+		{
+			Constraints[C]->RemoveStretch();
+		}
+
 		for (FEffector& Effector : Effectors)
 		{
 			Effector.SquashSubRoots(); // update squashing once again
@@ -160,6 +195,10 @@ void FPBIKSolver::Solve(const FPBIKSolverSettings& Settings)
 	for (const FEffector& Effector : Effectors)
 	{
 		FBone* Bone = Effector.Bone;
+		if (Bone->bIsSolverRoot)
+		{
+			continue; // if there's an effector on the root, leave it where the body ended up
+		}
 		Bone->Position = Bone->Parent->Position + Bone->Parent->Rotation * Bone->LocalPositionOrig;
 		Bone->Rotation = Effector.Rotation;
 	}
@@ -276,19 +315,6 @@ bool FPBIKSolver::InitBones()
 		}
 	}
 
-	// initialize local bone transforms
-	for (FBone& Bone : Bones)
-	{
-		FBone* Parent = Bone.Parent;
-		if (!Parent)
-		{
-			continue;
-		}
-
-		Bone.LocalPositionOrig = Parent->Rotation.Inverse() * (Bone.Position - Parent->Position);
-		Bone.LocalRotationOrig = Parent->Rotation.Inverse() * Bone.Rotation;
-	}
-
 	// initialize children lists
 	for (FBone& Parent : Bones)
 	{
@@ -307,24 +333,6 @@ bool FPBIKSolver::InitBones()
 		Bone.bIsSubRoot = Bone.Children.Num() > 1 || Bone.bIsSolverRoot;
 	}
 
-	// initialize Effector's nearest ParentSubRoot (FBone) pointer
-	// must be done AFTER setting: Bone.IsSubRoot/IsSolverRoot/Parent
-	for (FEffector& Effector : Effectors)
-	{
-		FBone* Parent = Effector.Bone->Parent;
-		while (Parent)
-		{
-			if (Parent->bIsSubRoot || Parent->bIsSolverRoot)
-			{
-				Effector.ParentSubRoot = Parent;
-				Effector.DistToSubRootOrig = (Parent->Position - Effector.Bone->Position).Size();
-				break;
-			}
-
-			Parent = Parent->Parent;
-		}
-	}
-
 	return true;
 }
 
@@ -335,15 +343,15 @@ bool FPBIKSolver::InitBodies()
 	using PBIK::FBone;
 
 	Bodies.Empty();
-
+	
 	// create bodies
 	for (FEffector& Effector : Effectors)
 	{
 		FBone* NextBone = Effector.Bone;
 		while (true)
 		{
-			FBone* BodyBone = NextBone->Parent;
-			if (!ensureMsgf(BodyBone, TEXT("PBIK: effector is on bone that is not child of root bone.")))
+			FBone* BodyBone = NextBone->bIsSolverRoot ? NextBone : NextBone->Parent;
+			if (!ensureMsgf(BodyBone, TEXT("PBIK: effector is on bone that is not on or below root bone.")))
 			{
 				return false;
 			}
@@ -374,6 +382,29 @@ bool FPBIKSolver::InitBodies()
 		Body.Bone->Body = &Body;
 	}
 
+	// initialize Effector's nearest ParentSubRoot (FBody) pointer
+	// must be done AFTER setting: Bone.IsSubRoot/IsSolverRoot/Parent
+	for (FEffector& Effector : Effectors)
+	{
+		FBone* Parent = Effector.Bone->Parent;
+		while (Parent)
+		{
+			if (!Parent->bIsSolved)
+			{
+				break; // this only happens when effector is on solver root
+			}
+
+			if (Parent->bIsSubRoot || Parent->bIsSolverRoot)
+			{
+				Effector.ParentSubRoot = Parent->Body;
+				Effector.DistToSubRootOrig = (Parent->Body->Position - Effector.Bone->Position).Size();
+				break;
+			}
+
+			Parent = Parent->Parent;
+		}
+	}
+
 	return true;
 }
 
@@ -392,32 +423,38 @@ void FPBIKSolver::AddBodyForBone(PBIK::FBone* Bone)
 void FPBIKSolver::InitConstraints()
 {
 	using PBIK::FEffector;
+	using PBIK::FBone;
 	using PBIK::FRigidBody;
 	using PBIK::FPinConstraint;
 	using PBIK::FJointConstraint;
 
 	Constraints.Empty();
 
-	// pin root body to animated location (usually disabled by solver settings)
-	TSharedPtr<FPinConstraint> RootConstraint = MakeShared<FPinConstraint>(SolverRoot->Body, SolverRoot->Position);
-	Constraints.Add(RootConstraint);
-	RootPin = RootConstraint;
-
 	// pin bodies to effectors
 	for (FEffector& Effector : Effectors)
 	{
-		if (!ensureMsgf(Effector.Bone->Parent, TEXT("PBIK: effector is on bone that does not have a parent.")))
+		FBone* BodyBone = Effector.Bone->bIsSolverRoot ? Effector.Bone : Effector.Bone->Parent;
+		if (!ensureMsgf(BodyBone, TEXT("PBIK: effector is on bone that does not have a parent.")))
 		{
 			return;
 		}
 
-		FRigidBody* Body = Effector.Bone->Parent->Body;
+		FRigidBody* Body = BodyBone->Body;
 		TSharedPtr<FPinConstraint> Constraint = MakeShared<FPinConstraint>(Body, Effector.Position);
 		Effector.Pin = Constraint;
-		Body->bPinnedToEffector = true;
+		Body->AttachedEffector = &Effector;
 		Constraints.Add(Constraint);
 	}
 
+	// pin root body to animated location 
+	// this constraint is by default off in solver settings
+	if (!SolverRoot->Body->AttachedEffector) // only add if user hasn't added their own root effector
+	{
+		TSharedPtr<FPinConstraint> RootConstraint = MakeShared<FPinConstraint>(SolverRoot->Body, SolverRoot->Position);
+		Constraints.Add(RootConstraint);
+		RootPin = RootConstraint;
+	}
+
 	// constrain all bodies together (child to parent)
 	for (FRigidBody& Body : Bodies)
 	{
@@ -464,18 +501,18 @@ int32 FPBIKSolver::AddBone(
 	return Bones.Emplace(Name, ParentIndex, InOrigPosition, InOrigRotation, bIsSolverRoot);
 }
 
-bool FPBIKSolver::AddEffector(FName BoneName)
+int32 FPBIKSolver::AddEffector(FName BoneName)
 {
 	for (PBIK::FBone& Bone : Bones)
 	{
 		if (Bone.Name == BoneName)
 		{
 			Effectors.Emplace(&Bone);
-			return true;
+			return Effectors.Num() - 1;
 		}
 	}
 
-	return false;
+	return -1;
 }
 
 void FPBIKSolver::SetBoneTransform(
@@ -521,8 +558,9 @@ void FPBIKSolver::SetEffectorGoal(
 	const int32 Index,
 	const FVector& InPosition,
 	const FQuat& InRotation,
-	const float Alpha)
+	const float OffsetAlpha,
+	const float StrengthAlpha)
 {
 	check(Index >= 0 && Index < Effectors.Num());
-	Effectors[Index].SetGoal(InPosition, InRotation, Alpha);
+	Effectors[Index].SetGoal(InPosition, InRotation, OffsetAlpha, StrengthAlpha);
 }
diff --git a/Engine/Plugins/Experimental/Animation/PBIK/Source/PBIK/Private/IKRig_PBIK.cpp b/Engine/Plugins/Experimental/Animation/PBIK/Source/PBIK/Private/IKRig_PBIK.cpp
new file mode 100644
--- /dev/null
+++ b/Engine/Plugins/Experimental/Animation/PBIK/Source/PBIK/Private/IKRig_PBIK.cpp
@@ -0,0 +1,65 @@
+// Copyright Epic Games, Inc. All Rights Reserved.
+
+#include "IKRig_PBIK.h"
+
+
+void UIKRig_PBIK::Init(const FIKRigTransforms& InGlobalTransforms)
+{
+	/*
+	// check how many effectors are assigned to a bone
+	int NumEffectors = 0;
+	for (const FPBIKRigEffector& Effector : Effectors)
+	{
+		if (Bones.GetIndex(Effector.Bone) != -1)
+		{
+			++NumEffectors; // bone is set and exists!
+		}
+	}
+
+	// validate inputs are ready to be initialized
+	bool bHasEffectors = NumEffectors > 0;
+	bool bRootIsAssigned = Root != NAME_None;
+	if (!(bHasEffectors && bRootIsAssigned))
+	{
+		return; // not setup yet
+	}
+
+	// reset all internal data
+	Solver.Reset();
+
+	// create bones
+	for (int B = 0; B < Bones.Num(); ++B)
+	{
+		FName Name = Bones[B].Name;
+		int ParentIndex = Bones[B].ParentIndex;
+		FVector InOrigPosition = Bones[B].InitialTransform.GetLocation();
+		FQuat InOrigRotation = Bones[B].InitialTransform.GetRotation();
+		bool bIsRoot = Bones[B].Name == Root;
+		Solver.AddBone(Name, ParentIndex, InOrigPosition, InOrigRotation, bIsRoot);
+	}
+
+	// create effectors
+	EffectorSolverIndices.Reset();
+	for (const FPBIKEffector& Effector : Effectors)
+	{
+		int32 IndexInSolver = Solver.AddEffector(Effector.Bone);
+		EffectorSolverIndices.Add(IndexInSolver);
+	}
+
+	// initialize
+	Solver.Initialize();
+
+	return; // don't update during init
+	*/
+}
+
+void UIKRig_PBIK::Solve(FIKRigTransforms& InOutGlobalTransforms, const FIKRigGoalContainer& Goals, FControlRigDrawInterface* InOutDrawInterface)
+{
+
+}
+
+void UIKRig_PBIK::CollectGoalNames(TSet<FName>& OutGoals) const
+{
+
+}
+
diff --git a/Engine/Plugins/Experimental/Animation/PBIK/Source/PBIK/Private/PBIK.cpp b/Engine/Plugins/Experimental/Animation/PBIK/Source/PBIK/Private/PBIK.cpp
--- a/Engine/Plugins/Experimental/Animation/PBIK/Source/PBIK/Private/PBIK.cpp
+++ b/Engine/Plugins/Experimental/Animation/PBIK/Source/PBIK/Private/PBIK.cpp
@@ -4,7 +4,7 @@
 
 #define LOCTEXT_NAMESPACE "FPBIKModule"
 
-DEFINE_LOG_CATEGORY(LogIKSolver);
+DEFINE_LOG_CATEGORY(LogPBIKSolver);
 
 void FPBIKModule::StartupModule()
 {
diff --git a/Engine/Plugins/Experimental/Animation/PBIK/Source/PBIK/Private/RigUnit_PBIK.cpp b/Engine/Plugins/Experimental/Animation/PBIK/Source/PBIK/Private/RigUnit_PBIK.cpp
--- a/Engine/Plugins/Experimental/Animation/PBIK/Source/PBIK/Private/RigUnit_PBIK.cpp
+++ b/Engine/Plugins/Experimental/Animation/PBIK/Source/PBIK/Private/RigUnit_PBIK.cpp
@@ -3,6 +3,8 @@
 #include "RigUnit_PBIK.h"
 #include "Units/RigUnitContext.h"
 
+//#pragma optimize("", off)
+
 FRigUnit_PBIK_Execute()
 {
 	DECLARE_SCOPE_HIERARCHICAL_COUNTER_RIGUNIT()
@@ -21,7 +23,7 @@ FRigUnit_PBIK_Execute()
 		int NumEffectors = 0;
 		for (const FPBIKEffector& Effector : Effectors)
 		{
-			if (Bones.GetIndex(Effector.Bone) != NAME_None)
+			if (Bones.GetIndex(Effector.Bone) != INDEX_NONE)
 			{
 				++NumEffectors; // bone is set and exists!
 			}
@@ -50,12 +52,11 @@ FRigUnit_PBIK_Execute()
 		}
 		
 		// create effectors
+		EffectorSolverIndices.Reset();
 		for (const FPBIKEffector& Effector : Effectors)
 		{
-			if (!Solver.AddEffector(Effector.Bone))
-			{
-				return;
-			}
+			int32 IndexInSolver = Solver.AddEffector(Effector.Bone);
+			EffectorSolverIndices.Add(IndexInSolver);
 		}
 		
 		// initialize
@@ -79,7 +80,7 @@ FRigUnit_PBIK_Execute()
 	for (const FPBIKBoneSetting& BoneSetting : BoneSettings)
 	{
 		int32 BoneIndex = Bones.GetIndex(BoneSetting.Bone);
-		if (BoneIndex == INDEX_NONE)
+		if (BoneIndex == -1)
 		{
 			continue; // no bones to apply it to
 		}
@@ -93,18 +94,27 @@ FRigUnit_PBIK_Execute()
 	// update effectors
 	for (int E = 0; E < Effectors.Num(); ++E)
 	{
-		FVector Position = Effectors[E].Transform.GetLocation();
-		FQuat Rotation = Effectors[E].Transform.GetRotation();
-		Solver.SetEffectorGoal(E, Position, Rotation, Effectors[E].Alpha);
+		if (EffectorSolverIndices[E] == -1)
+		{
+			continue;
+		}
+
+		const FPBIKEffector& Effector = Effectors[E];
+		FVector Position = Effector.Transform.GetLocation();
+		FQuat Rotation = Effector.Transform.GetRotation();
+		Solver.SetEffectorGoal(EffectorSolverIndices[E], Position, Rotation, Effector.OffsetAlpha, Effector.StrengthAlpha);
 	}
 
 	// solve
 	Solver.Solve(Settings);
 
 	// copy transforms back
-	for (int B = 0; B < Bones.Num(); ++B)
+	bool bPropagateTransform = false;
+	for (int BoneIndex = 0; BoneIndex < Bones.Num(); ++BoneIndex)
 	{
-		Solver.GetBoneGlobalTransform(B, Bones[B].GlobalTransform);
+		FTransform NewTransform;
+		Solver.GetBoneGlobalTransform(BoneIndex, NewTransform);
+		Bones.SetGlobalTransform(BoneIndex, NewTransform, bPropagateTransform);
 	}
 
 	// do all debug drawing
diff --git a/Engine/Plugins/Experimental/Animation/PBIK/Source/PBIK/Public/Core/PBIKBody.h b/Engine/Plugins/Experimental/Animation/PBIK/Source/PBIK/Public/Core/PBIKBody.h
--- a/Engine/Plugins/Experimental/Animation/PBIK/Source/PBIK/Public/Core/PBIKBody.h
+++ b/Engine/Plugins/Experimental/Animation/PBIK/Source/PBIK/Public/Core/PBIKBody.h
@@ -6,7 +6,6 @@
 #include "Math/Vector.h"
 #include "Math/Quat.h"
 
-
 struct FPBIKSolverSettings;
 
 namespace PBIK
@@ -14,6 +13,7 @@ namespace PBIK
 
 struct FJointConstraint;
 struct FRigidBody;
+struct FEffector;
 
 struct FBone
 {
@@ -83,7 +83,7 @@ struct FRigidBody
 	TArray<FVector> ChildLocalPositions;
 
 	float InvMass = 0.0f;
-	bool bPinnedToEffector = false;
+	FEffector* AttachedEffector = nullptr;
 	float Length;
 	
 private:
diff --git a/Engine/Plugins/Experimental/Animation/PBIK/Source/PBIK/Public/Core/PBIKSolver.h b/Engine/Plugins/Experimental/Animation/PBIK/Source/PBIK/Public/Core/PBIKSolver.h
--- a/Engine/Plugins/Experimental/Animation/PBIK/Source/PBIK/Public/Core/PBIKSolver.h
+++ b/Engine/Plugins/Experimental/Animation/PBIK/Source/PBIK/Public/Core/PBIKSolver.h
@@ -12,7 +12,10 @@ DECLARE_CYCLE_STAT(TEXT("PBIK Solve"), STAT_PBIK_Solve, STATGROUP_Anim);
 
 namespace PBIK
 {
-	static float GLOBAL_UNITS = 100.0f; // (1.0f = meters), (100.0f = centimeters)
+
+static float GLOBAL_UNITS = 100.0f; // (1.0f = meters), (100.0f = centimeters)
+
+struct FRigidBody;
 
 struct FEffector
 {
@@ -29,13 +32,14 @@ struct FEffector
 	TWeakPtr<FPinConstraint> Pin;
 	
 	float DistToSubRootOrig;
-	FBone* ParentSubRoot;
+	FRigidBody* ParentSubRoot = nullptr;
 
-	float Alpha;
+	float TransformAlpha;
+	float StrengthAlpha;
 
 	FEffector(FBone* InBone);
 
-	void SetGoal( const FVector InPositionGoal, const FQuat& InRotationGoal, float InAlpha);
+	void SetGoal(const FVector InPositionGoal, const FQuat& InRotationGoal, float InTransformAlpha, float InStrengthAlpha);
 
 	void UpdateFromInputs();
 
@@ -49,16 +53,16 @@ struct FPBIKSolverSettings
 {
 	GENERATED_BODY()
 
-	UPROPERTY(meta = (ClampMin = "0", UIMin = "0.0", UIMax = "200.0"))
+	UPROPERTY(EditAnywhere, Category = SolverSettings, meta = (ClampMin = "0", UIMin = "0.0", UIMax = "200.0"))
 	int32 Iterations = 20;
 
-	UPROPERTY(meta = (ClampMin = "0", UIMin = "0.0", UIMax = "10.0"))
+	UPROPERTY(EditAnywhere, Category = SolverSettings, meta = (ClampMin = "0", UIMin = "0.0", UIMax = "10.0"))
 	float MassMultiplier = 1.0f;
 
-	UPROPERTY()
+	UPROPERTY(EditAnywhere, Category = SolverSettings)
 	bool bAllowStretch = false;
 
-	UPROPERTY()
+	UPROPERTY(EditAnywhere, Category = SolverSettings)
 	bool bPinRoot = false;
 };
 
@@ -93,7 +97,12 @@ public:
 
 	PBIK::FBoneSettings* GetBoneSettings(const int32 Index);
 
-	void SetEffectorGoal(const int32 Index, const FVector& InPosition, const FQuat& InRotation, const float Alpha);
+	void SetEffectorGoal(
+		const int32 Index, 
+		const FVector& InPosition, 
+		const FQuat& InRotation, 
+		const float OffsetAlpha, 
+		const float StrengthAlpha);
 
 	void GetBoneGlobalTransform(const int32 Index, FTransform& OutTransform);
 
@@ -108,7 +117,7 @@ public:
 		const FQuat& InOrigRotation,
 		bool bIsSolverRoot);
 
-	bool AddEffector(FName BoneName);
+	int32 AddEffector(FName BoneName);
 	
 private:
 
diff --git a/Engine/Plugins/Experimental/Animation/PBIK/Source/PBIK/Public/IKRig_PBIK.h b/Engine/Plugins/Experimental/Animation/PBIK/Source/PBIK/Public/IKRig_PBIK.h
new file mode 100644
--- /dev/null
+++ b/Engine/Plugins/Experimental/Animation/PBIK/Source/PBIK/Public/IKRig_PBIK.h
@@ -0,0 +1,55 @@
+// Copyright Epic Games, Inc. All Rights Reserved.
+
+#pragma once
+
+#include "IKRigSolver.h"
+
+#include "Core/PBIKSolver.h"
+#include "PBIK_Shared.h"
+
+#include "IKRig_PBIK.generated.h"
+
+
+USTRUCT(BlueprintType)
+struct FPBIKRigEffector
+{
+	GENERATED_BODY()
+
+	UPROPERTY(EditAnywhere, Category = FFBIKRigEffector)
+	FIKRigEffectorGoal Target;
+
+	UPROPERTY(EditAnywhere, Category = FFBIKRigEffector)
+	float OffsetAlpha = 1.0f;
+
+	UPROPERTY(EditAnywhere, Category = FFBIKRigEffector)
+	float StrengthAlpha = 1.0f;
+};
+
+
+UCLASS(EditInlineNew)
+class PBIK_API UIKRig_PBIK : public UIKRigSolver
+{
+	GENERATED_BODY()
+
+	UPROPERTY(EditAnywhere, Category = Settings)
+	FName Root;
+
+	UPROPERTY(EditAnywhere, Category = Settings)
+	TArray<FPBIKRigEffector> Effectors;
+	UPROPERTY(transient)
+	TArray<int32> EffectorSolverIndices;
+
+	UPROPERTY(EditAnywhere, Category = Settings)
+	TArray<FPBIKBoneSetting> BoneSettings;
+
+	UPROPERTY(EditAnywhere, Category = Settings)
+	FPBIKSolverSettings SolverSettings;
+
+protected:
+	virtual void Init(const FIKRigTransforms& InGlobalTransforms) override;
+	virtual void Solve(
+		FIKRigTransforms& InOutGlobalTransforms,
+		const FIKRigGoalContainer& Goals,
+		FControlRigDrawInterface* InOutDrawInterface) override;
+	virtual void CollectGoalNames(TSet<FName>& OutGoals) const override;
+};
\ No newline at end of file
diff --git a/Engine/Plugins/Experimental/Animation/PBIK/Source/PBIK/Public/PBIK.h b/Engine/Plugins/Experimental/Animation/PBIK/Source/PBIK/Public/PBIK.h
--- a/Engine/Plugins/Experimental/Animation/PBIK/Source/PBIK/Public/PBIK.h
+++ b/Engine/Plugins/Experimental/Animation/PBIK/Source/PBIK/Public/PBIK.h
@@ -14,4 +14,4 @@ public:
 	virtual void ShutdownModule() override;
 };
 
-DECLARE_LOG_CATEGORY_EXTERN(LogIKSolver, Warning, All);
+DECLARE_LOG_CATEGORY_EXTERN(LogPBIKSolver, Warning, All);
diff --git a/Engine/Plugins/Experimental/Animation/PBIK/Source/PBIK/Public/PBIK_Shared.h b/Engine/Plugins/Experimental/Animation/PBIK/Source/PBIK/Public/PBIK_Shared.h
new file mode 100644
--- /dev/null
+++ b/Engine/Plugins/Experimental/Animation/PBIK/Source/PBIK/Public/PBIK_Shared.h
@@ -0,0 +1,76 @@
+// Copyright Epic Games, Inc. All Rights Reserved.
+
+#pragma once
+
+#include "CoreMinimal.h"
+
+#include "Core/PBIKSolver.h"
+
+#include "PBIK_Shared.generated.h"
+
+UENUM(BlueprintType)
+enum class EPBIKLimitType : uint8
+{
+	Free,
+	Limited,
+	Locked,
+};
+
+USTRUCT(BlueprintType)
+struct FPBIKBoneSetting
+{
+	GENERATED_BODY()
+
+	FPBIKBoneSetting() : Bone(NAME_None) {}
+
+	UPROPERTY(meta = (Constant, CustomWidget = "BoneName"))
+	FName Bone;
+
+	UPROPERTY(meta = (ClampMin = "0", ClampMax = "1", UIMin = "0.0", UIMax = "1.0"))
+	float RotationStiffness = 0.0f;
+	UPROPERTY(meta = (ClampMin = "0", ClampMax = "1", UIMin = "0.0", UIMax = "1.0"))
+	float PositionStiffness = 0.0f;
+
+	UPROPERTY()
+	EPBIKLimitType X;
+	UPROPERTY(meta = (ClampMin = "-180", ClampMax = "0", UIMin = "-180.0", UIMax = "0.0"))
+	float MinX = 0.0f;
+	UPROPERTY(meta = (ClampMin = "0", ClampMax = "180", UIMin = "0.0", UIMax = "180.0"))
+	float MaxX = 0.0f;
+
+	UPROPERTY()
+	EPBIKLimitType Y;
+	UPROPERTY(meta = (ClampMin = "-180", ClampMax = "0", UIMin = "-180.0", UIMax = "0.0"))
+	float MinY = 0.0f;
+	UPROPERTY(meta = (ClampMin = "0", ClampMax = "180", UIMin = "0.0", UIMax = "180.0"))
+	float MaxY = 0.0f;
+
+	UPROPERTY()
+	EPBIKLimitType Z;
+	UPROPERTY(meta = (ClampMin = "-180", ClampMax = "0", UIMin = "-180.0", UIMax = "0.0"))
+	float MinZ = 0.0f;
+	UPROPERTY(meta = (ClampMin = "0", ClampMax = "180", UIMin = "0.0", UIMax = "180.0"))
+	float MaxZ = 0.0f;
+
+	UPROPERTY()
+	bool bUsePreferredAngles = false;
+	UPROPERTY()
+	FRotator PreferredAngles;
+
+	void CopyToCoreStruct(PBIK::FBoneSettings& Settings) const
+	{
+		Settings.RotationStiffness = RotationStiffness;
+		Settings.PositionStiffness = PositionStiffness;
+		Settings.X = static_cast<PBIK::ELimitType>(X);
+		Settings.MinX = MinX;
+		Settings.MaxX = MaxX;
+		Settings.Y = static_cast<PBIK::ELimitType>(Y);
+		Settings.MinY = MinY;
+		Settings.MaxY = MaxY;
+		Settings.Z = static_cast<PBIK::ELimitType>(Z);
+		Settings.MinZ = MinZ;
+		Settings.MaxZ = MaxZ;
+		Settings.bUsePreferredAngles = bUsePreferredAngles;
+		Settings.PreferredAngles = PreferredAngles;
+	}
+};
\ No newline at end of file
diff --git a/Engine/Plugins/Experimental/Animation/PBIK/Source/PBIK/Public/RigUnit_PBIK.h b/Engine/Plugins/Experimental/Animation/PBIK/Source/PBIK/Public/RigUnit_PBIK.h
--- a/Engine/Plugins/Experimental/Animation/PBIK/Source/PBIK/Public/RigUnit_PBIK.h
+++ b/Engine/Plugins/Experimental/Animation/PBIK/Source/PBIK/Public/RigUnit_PBIK.h
@@ -8,6 +8,8 @@
 #include "Core/PBIKSolver.h"
 #include "Core/PBIKDebug.h"
 
+#include "PBIK_Shared.h"
+
 #include "RigUnit_PBIK.generated.h"
 
 using PBIK::FDebugLine;
@@ -57,74 +59,10 @@ struct FPBIKEffector
 	FTransform Transform;
 
 	UPROPERTY()
-	float Alpha = 1.0f;
-};
-
-UENUM(BlueprintType)
-enum class EPBIKLimitType : uint8
-{
-	Free,
-	Limited,
-	Locked,
-};
-
-USTRUCT(BlueprintType)
-struct FPBIKBoneSetting
-{
-	GENERATED_BODY()
-
-	FPBIKBoneSetting() : Bone(NAME_None) {}
-
-	UPROPERTY(meta = (Constant, CustomWidget = "BoneName"))
-	FName Bone;
+	float OffsetAlpha = 1.0f;
 
-	UPROPERTY(meta = (ClampMin = "0", ClampMax = "1", UIMin = "0.0", UIMax = "1.0"))
-	float RotationStiffness = 0.0f;
-	UPROPERTY(meta = (ClampMin = "0", ClampMax = "1", UIMin = "0.0", UIMax = "1.0"))
-	float PositionStiffness = 0.0f;
-
-	UPROPERTY()
-	EPBIKLimitType X;
-	UPROPERTY(meta = (ClampMin = "-180", ClampMax = "0", UIMin = "-180.0", UIMax = "0.0"))
-	float MinX = 0.0f;
-	UPROPERTY(meta = (ClampMin = "0", ClampMax = "180", UIMin = "0.0", UIMax = "180.0"))
-	float MaxX = 0.0f;
-
-	UPROPERTY()
-	EPBIKLimitType Y;
-	UPROPERTY(meta = (ClampMin = "-180", ClampMax = "0", UIMin = "-180.0", UIMax = "0.0"))
-	float MinY = 0.0f;
-	UPROPERTY(meta = (ClampMin = "0", ClampMax = "180", UIMin = "0.0", UIMax = "180.0"))
-	float MaxY = 0.0f;
-
-	UPROPERTY()
-	EPBIKLimitType Z;
-	UPROPERTY(meta = (ClampMin = "-180", ClampMax = "0", UIMin = "-180.0", UIMax = "0.0"))
-	float MinZ = 0.0f;
-	UPROPERTY(meta = (ClampMin = "0", ClampMax = "180", UIMin = "0.0", UIMax = "180.0"))
-	float MaxZ = 0.0f;
-
-	UPROPERTY()
-	bool bUsePreferredAngles = false;
 	UPROPERTY()
-	FRotator PreferredAngles;
-
-	void CopyToCoreStruct(PBIK::FBoneSettings& Settings) const
-	{
-		Settings.RotationStiffness = RotationStiffness;
-		Settings.PositionStiffness = PositionStiffness;
-		Settings.X = static_cast<PBIK::ELimitType>(X);
-		Settings.MinX = MinX;
-		Settings.MaxX = MaxX;
-		Settings.Y = static_cast<PBIK::ELimitType>(Y);
-		Settings.MinY = MinY;
-		Settings.MaxY = MaxY;
-		Settings.Z = static_cast<PBIK::ELimitType>(Z);
-		Settings.MinZ = MinZ;
-		Settings.MaxZ = MaxZ;
-		Settings.bUsePreferredAngles = bUsePreferredAngles;
-		Settings.PreferredAngles = PreferredAngles;
-	}
+	float StrengthAlpha = 1.0f;
 };
 
 USTRUCT(meta=(DisplayName="Position Based IK", Category="Hierarchy", Keywords="Position Based, PBIK, IK, Full Body"))
@@ -145,6 +83,8 @@ struct FRigUnit_PBIK : public FRigUnit_HighlevelBaseMutable
 
 	UPROPERTY(meta = (Input))
 	TArray<FPBIKEffector> Effectors;
+	UPROPERTY(transient)
+	TArray<int32> EffectorSolverIndices;
 
 	UPROPERTY(meta = (Input))
 	TArray<FPBIKBoneSetting> BoneSettings;
